Code and documentation cleanup pass

diff --git a/Engine/Source/Editor/InputBindingEditor/Private/Widgets/SGestureEditor.h b/Engine/Source/Editor/InputBindingEditor/Private/Widgets/SGestureEditor.h
--- a/Engine/Source/Editor/InputBindingEditor/Private/Widgets/SGestureEditor.h
+++ b/Engine/Source/Editor/InputBindingEditor/Private/Widgets/SGestureEditor.h
@@ -1,14 +1,10 @@
 // Copyright 1998-2014 Epic Games, Inc. All Rights Reserved.
 
-/*=============================================================================
-	SGestureEditor.h: Declares the SGestureEditor class.
-=============================================================================*/
-
 #pragma once
 
 
 /**
- * A specialized text edit box that visualizes a new gesture being entered           
+ * A specialized text edit box that visualizes a new gesture being entered           .
  */
 class SGestureEditor
 	: public SEditableText
@@ -22,89 +18,79 @@ public:
 		SLATE_EVENT( FSimpleDelegate, OnEditingStarted )
 	SLATE_END_ARGS()
 
-
 public:
 
-	/**
-	 * Default constructor.
-	 */
+	/** Default constructor. */
 	SGestureEditor()
 		: bIsEditing( false )
 		, bIsTyping( false )
 	{ }
 
-
 public:
 
 	/**
 	 * Constructs the widget.
 	 *
-	 * @param InArgs - The Slate argument list.
-	 * @param InputCommand -
+	 * @param InArgs The Slate argument list.
+	 * @param InputCommand
 	 */
 	void Construct( const FArguments& InArgs, TSharedPtr<FGestureTreeItem> InputCommand );
 	
 	virtual bool SupportsKeyboardFocus() const { return true; }
 
-	/** Starts editing the gesture */
+	/** Starts editing the gesture. */
 	void StartEditing();
 	
-	/** Stops editing the gesture */
+	/** Stops editing the gesture. */
 	void StopEditing();
 
-	/** Commits the new gesture to the commands active gesture */
+	/** Commits the new gesture to the commands active gesture. */
 	void CommitNewGesture();
 
-	/** Removes the active gesture from the command */
+	/** Removes the active gesture from the command. */
 	void RemoveActiveGesture();
 
-	/** @return Whether or not we are in editing mode */
+	/** @return Whether or not we are in editing mode. */
 	bool IsEditing() const { return bIsEditing; }
 
-	/** @return True if the user is physically typing a key */
+	/** @return True if the user is physically typing a key. */
 	bool IsTyping() const { return bIsTyping; }
 
-	/** @return Whether or not the gesture being edited is valid */
+	/** @return Whether or not the gesture being edited is valid. */
 	bool IsEditedGestureValid() const { return EditingInputGesture.IsValidGesture(); }
 
-	/** @return Whether or not the command has a valid gesture */
+	/** @return Whether or not the command has a valid gesture. */
 	bool IsActiveGestureValid() const { return CommandInfo->GetActiveGesture()->IsValidGesture(); }
 
-	/** @return the Notification message being displayed if any */
+	/** @return the Notification message being displayed if any. */
 	const FText& GetNotificationText() const { return NotificationMessage; };
 
-	/** @return true if the edited gesture has a conflict with an existing gesture */
+	/** @return true if the edited gesture has a conflict with an existing gesture. */
 	bool HasConflict() const { return !NotificationMessage.IsEmpty(); }
 
-
 private:
 
 	/** Gesture that is currently being edited. */
 	static TWeakPtr<SGestureEditor> GestureBeingEdited;
 
 	virtual FReply OnKeyDown( const FGeometry& MyGeometry, const FKeyboardEvent& InKeyboardEvent ) override;
-
 	virtual FReply OnKeyUp( const FGeometry& MyGeometry, const FKeyboardEvent& InKeyboardEvent ) override;
-
 	virtual FReply OnMouseButtonDown( const FGeometry& MyGeometry, const FPointerEvent& InMouseEvent ) override;
-
 	virtual FReply OnMouseButtonDoubleClick( const FGeometry& InMyGeometry, const FPointerEvent& InMouseEvent ) override;
-
 	virtual void OnKeyboardFocusLost( const FKeyboardFocusEvent& InKeyboardFocusEvent ) override;
-
 	virtual FReply OnKeyChar( const FGeometry& MyGeometry, const FCharacterEvent& InCharacterEvent ) override;
 
 	/** 
-	 * Called when the gesture changes 
+	 * Called when the gesture changes.
 	 *
-	 * @param NewGesture	The new gesture 
+	 * @param NewGesture The new gesture.
 	 */
 	void OnGestureTyped( const FInputGesture& NewGesture );
 
 	/** 
-	 * Called when the gesture changes 
+	 * Called when the gesture changes.
 	 *
-	 * @param NewGesture	The gesture to commit
+	 * @param NewGesture The gesture to commit.
 	 */
 	void OnGestureCommitted( const FInputGesture& NewGesture );
 
@@ -114,33 +100,32 @@ private:
 	/** @return The hint text to display in the text box if it is empty */
 	FText OnGetGestureInputHintText() const;
 
-
 private:
 
-	// The command we are editing a gesture for.
+	/** The command we are editing a gesture for. */
 	TSharedPtr<FUICommandInfo> CommandInfo;
 
-	// Delegate to execute when the edit box loses focus.
+	/** Delegate to execute when the edit box loses focus. */
 	FSimpleDelegate OnEditBoxLostFocus;
 
-	// Delegate to execute when the gesture changes.
+	/** Delegate to execute when the gesture changes. */
 	FSimpleDelegate OnGestureChanged;
 
-	// Delegate to execute when we stop editing.
+	/** Delegate to execute when we stop editing. */
 	FSimpleDelegate OnEditingStopped;
 
-	// Delegate to execute when we start editing.
+	/** Delegate to execute when we start editing. */
 	FSimpleDelegate OnEditingStarted;
 
-	// The notification message (duplicate bindings) being displayed.
+	/** The notification message (duplicate bindings) being displayed. */
 	FText NotificationMessage;
 
-	// Temp gesture being edited.
+	/** Temp gesture being edited. */
 	FInputGesture EditingInputGesture;
 
-	// Whether or not we are in edit mode.
+	/** Whether or not we are in edit mode. */
 	bool bIsEditing;
 
-	// Whether or not the user is physically typing a new key 
+	/** Whether or not the user is physically typing a new key. */
 	bool bIsTyping;
 };
diff --git a/Engine/Source/Editor/InputBindingEditor/Private/Widgets/SGestureTreeItem.h b/Engine/Source/Editor/InputBindingEditor/Private/Widgets/SGestureTreeItem.h
--- a/Engine/Source/Editor/InputBindingEditor/Private/Widgets/SGestureTreeItem.h
+++ b/Engine/Source/Editor/InputBindingEditor/Private/Widgets/SGestureTreeItem.h
@@ -1,9 +1,5 @@
 // Copyright 1998-2014 Epic Games, Inc. All Rights Reserved.
 
-/*=============================================================================
-	SGestureTreeItem.h: Declares the SGestureTreeItem class.
-=============================================================================*/
-
 #pragma once
 
 
@@ -27,7 +23,7 @@ typedef STreeView< TSharedPtr<FGestureTreeItem> > SGestureTree;
 
 
 /**
- * A widget which visualizes a command info                    
+ * A widget which visualizes a command info      .              
  */
 class SGestureTreeItem
 	: public SMultiColumnTableRow< TSharedPtr< FGestureTreeItem > >
@@ -37,15 +33,14 @@ public:
 	SLATE_BEGIN_ARGS( SGestureTreeItem ){}
 	SLATE_END_ARGS()
 
-
  public:
 
 	/**
 	 * Constructs the widget.
 	 *
-	 * @param InArgs - The Slate argument list.
-	 * @param InOwnerTable - The table that owns this tree item.
-	 * @param InItem - The actual item to be displayed.
+	 * @param InArgs The Slate argument list.
+	 * @param InOwnerTable The table that owns this tree item.
+	 * @param InItem The actual item to be displayed.
 	 */
 	void Construct( const FArguments& InArgs, const TSharedRef<STableViewBase>& InOwnerTable, TSharedPtr<FGestureTreeItem> InItem )
 	{
@@ -55,9 +50,9 @@ public:
 	}
 
 	/**
-	 * Called to generate a widget for each column
+	 * Called to generate a widget for each column.
 	 *
-	 * @param ColumnName	The name of the column that needs a widget
+	 * @param ColumnName The name of the column that needs a widget.
 	 */
 	TSharedRef<SWidget> GenerateWidgetForColumn( const FName& ColumnName )
 	{
@@ -113,7 +108,10 @@ public:
 		}
 	}
 
-	/** SWidget Interface*/
+public:
+
+	// SWidget interface
+
 	virtual FReply OnMouseButtonDown( const FGeometry& MyGeometry, const FPointerEvent& MouseEvent ) override
 	{
 		if( MouseEvent.GetEffectingButton() == EKeys::LeftMouseButton && FMultiBoxSettings::IsInToolbarEditMode() )
@@ -133,8 +131,9 @@ public:
 
 		return FReply::Handled().BeginDragDrop( FUICommandDragDropOp::New( TreeItem->CommandInfo.ToSharedRef(), NAME_None, TempBuilder.MakeWidget(), FVector2D::ZeroVector ) );
 	}
+
 private:
 
-	// Holds the tree item being visualized.
+	/** Holds the tree item being visualized. */
 	TSharedPtr<FGestureTreeItem> TreeItem;
 };
diff --git a/Engine/Source/Editor/InputBindingEditor/Private/Widgets/SInputBindingEditorPanel.cpp b/Engine/Source/Editor/InputBindingEditor/Private/Widgets/SInputBindingEditorPanel.cpp
--- a/Engine/Source/Editor/InputBindingEditor/Private/Widgets/SInputBindingEditorPanel.cpp
+++ b/Engine/Source/Editor/InputBindingEditor/Private/Widgets/SInputBindingEditorPanel.cpp
@@ -1,11 +1,8 @@
 // Copyright 1998-2014 Epic Games, Inc. All Rights Reserved.
 
-/*=============================================================================
-	SGestureEditor.cpp: Implements the SGestureEditor class.
-=============================================================================*/
-
 #include "InputBindingEditorPrivatePCH.h"
 
+
 #define LOCTEXT_NAMESPACE "SInputBindingEditorPanel"
 
 
@@ -95,6 +92,7 @@ void SInputBindingEditorPanel::Construct(const FArguments& InArgs)
 	GestureTree->RequestTreeRefresh();
 }
 
+
 static bool StringPassesFilter( const FText& Text, const TArray<FString>& FilterStrings)
 {
 	//@todo Use localization safe comparison methods [10/11/2013 justin.sargent]
@@ -116,6 +114,7 @@ static bool StringPassesFilter( const FText& Text, const TArray<FString>& Filter
 	return bAcceptable;
 }
 
+
 void SInputBindingEditorPanel::OnSearchChanged( const FText& NewSearch )
 {
 	FilterStrings.Empty();
@@ -129,6 +128,7 @@ void SInputBindingEditorPanel::OnSearchChanged( const FText& NewSearch )
 	FilterVisibleContextList();
 }
 
+
 FReply SInputBindingEditorPanel::OnBindingColumnClicked()
 {
 	static bool bSortUp = false;
@@ -141,6 +141,7 @@ FReply SInputBindingEditorPanel::OnBindingColumnClicked()
 	return FReply::Handled();
 }
 
+
 FReply SInputBindingEditorPanel::OnNameColumnClicked()
 {
 	static bool bSortUp = false;
@@ -154,6 +155,7 @@ FReply SInputBindingEditorPanel::OnNameColumnClicked()
 	return FReply::Handled();
 }
 
+
 void SInputBindingEditorPanel::UpdateContextMasterList()
 {
 	ContextMasterList.Empty();
@@ -183,9 +185,9 @@ void SInputBindingEditorPanel::UpdateContextMasterList()
 	}
 
 	ContextVisibleList = ContextMasterList;
-
 }
 
+
 void SInputBindingEditorPanel::FilterVisibleContextList()
 {
 	if( FilterStrings.Num() == 0 )
@@ -233,6 +235,7 @@ void SInputBindingEditorPanel::FilterVisibleContextList()
 	GestureTree->RequestTreeRefresh();
 }
 
+
 void SInputBindingEditorPanel::OnCommandsChanged()
 {
 	UpdateContextMasterList();
@@ -240,6 +243,7 @@ void SInputBindingEditorPanel::OnCommandsChanged()
 	FilterVisibleContextList();
 }
 
+
 TSharedRef< ITableRow > SInputBindingEditorPanel::OnGenerateWidgetForTreeItem( TSharedPtr<FGestureTreeItem> InTreeItem, const TSharedRef<STableViewBase>& OwnerTable )
 {
 	if( InTreeItem->IsContext() )
@@ -304,11 +308,13 @@ void SInputBindingEditorPanel::OnGetChildrenForTreeItem( TSharedPtr<FGestureTree
 	}
 }
 
+
 bool SInputBindingEditorPanel::SupportsKeyboardFocus() const
 {
 	return SearchBox->SupportsKeyboardFocus();
 }
 
+
 FReply SInputBindingEditorPanel::OnKeyboardFocusReceived( const FGeometry& MyGeometry, const FKeyboardFocusEvent& InKeyboardFocusEvent )
 {
 	FReply Reply = FReply::Handled();
@@ -321,4 +327,5 @@ FReply SInputBindingEditorPanel::OnKeyboardFocusReceived( const FGeometry& MyGeo
 	return Reply;
 }
 
+
 #undef LOCTEXT_NAMESPACE
diff --git a/Engine/Source/Editor/InputBindingEditor/Private/Widgets/SInputBindingEditorPanel.h b/Engine/Source/Editor/InputBindingEditor/Private/Widgets/SInputBindingEditorPanel.h
--- a/Engine/Source/Editor/InputBindingEditor/Private/Widgets/SInputBindingEditorPanel.h
+++ b/Engine/Source/Editor/InputBindingEditor/Private/Widgets/SInputBindingEditorPanel.h
@@ -1,9 +1,5 @@
 // Copyright 1998-2014 Epic Games, Inc. All Rights Reserved.
 
-/*=============================================================================
-	SInputBindingEditorPanel.h: Declares the SInputBindingEditorPanel class.
-=============================================================================*/
-
 #pragma once
 
 
@@ -32,13 +28,12 @@ struct FGestureSort
 		}
 	}
 
-
 private:
 
-	// Whether or not to sort by name.  If false we sort by binding.
+	/** Whether or not to sort by name.  If false we sort by binding. */
 	bool bSortName;
 
-	// Whether or not to sort up.  If false we sort down.
+	/** Whether or not to sort up.  If false we sort down. */
 	bool bSortUp;
 };
 
@@ -52,101 +47,92 @@ public:
 
 	SLATE_BEGIN_ARGS(SInputBindingEditorPanel){}
 	SLATE_END_ARGS()
-	
 
 public:
 
-	/**
-	 * Default constructor.
-	 */
+	/** Default constructor. */
 	SInputBindingEditorPanel()
 		: GestureSortMode( true, false )
 	{ }
 
-	/**
-	 * Destructor	Saves the user defined bindings to disk when closed
-	 */
+	/** Destructor. */
 	virtual ~SInputBindingEditorPanel()
 	{
 		FInputBindingManager::Get().SaveInputBindings();
 		FBindingContext::CommandsChanged.RemoveAll( this );
 	}
 
-
 public:
 
 	/**
 	 * Constructs the widget.
 	 *
-	 * @param InArgs - The Slate argument list.
+	 * @param InArgs The Slate argument list.
 	 */
 	void Construct(const FArguments& InArgs);
 
+public:
+
 	// SWidget interface
+
 	virtual bool SupportsKeyboardFocus() const override;
 	virtual FReply OnKeyboardFocusReceived( const FGeometry& MyGeometry, const FKeyboardFocusEvent& InKeyboardFocusEvent ) override;
-	// End of SWidget interface
 
 private:
 
 	/**
-	 * Called when the text changes in the search box
+	 * Called when the text changes in the search box.
 	 * 
-	 * @param NewSearch	The new search term to filter by
+	 * @param NewSearch	The new search term to filter by.
 	 */
 	void OnSearchChanged( const FText& NewSearch );
 
 	/**
-	 * Generates widget for an item in the gesture tree
+	 * Generates widget for an item in the gesture tree.
 	 */
 	TSharedRef< ITableRow > OnGenerateWidgetForTreeItem( TSharedPtr<FGestureTreeItem> InTreeItem, const TSharedRef<STableViewBase>& OwnerTable );
 
 	/**
-	 * Gets children FGestureTreeItems from the passed in tree item.  Note: Only contexts have children and those children are the actual gestures
+	 * Gets children FGestureTreeItems from the passed in tree item.  Note: Only contexts have children and those children are the actual gestures.
 	 */
 	void OnGetChildrenForTreeItem( TSharedPtr<FGestureTreeItem> InTreeItem, TArray< TSharedPtr< FGestureTreeItem > >& OutChildren );
 
 	/**
-	 * Called when the binding column is clicked.  We sort by binding in this case 
+	 * Called when the binding column is clicked.  We sort by binding in this case .
 	 */	 
 	FReply OnBindingColumnClicked();
 
 	/**
-	 * Called when the name column is clicked. We sort by name in this case
+	 * Called when the name column is clicked. We sort by name in this case.
 	 */	
 	FReply OnNameColumnClicked();
 
-	/**
-	 * Updates the master context list with new commands
-	 */
+	/** Updates the master context list with new commands. */
 	void UpdateContextMasterList();
 
-	/**
-	 * Filters the currently visible context list
-	 */
+	/** Filters the currently visible context list. */
 	void FilterVisibleContextList();
 
-	/**
-	 * Called when new commands are registered with the input binding manager
-	 */
+	/** Called when new commands are registered with the input binding manager. */
 	void OnCommandsChanged();
+
 private:
 
-	// List of all known contexts.
+	/** List of all known contexts. */
 	TArray< TSharedPtr<FGestureTreeItem> > ContextMasterList;
 	
-	// List of contexts visible in the tree.
+	/** List of contexts visible in the tree. */
 	TArray< TSharedPtr<FGestureTreeItem> > ContextVisibleList;
 	
-	// Search box
+	/** Search box */
 	TSharedPtr<SWidget> SearchBox;
 
-	// Gesture tree widget.
+	/** Gesture tree widget. */
 	TSharedPtr< SGestureTree > GestureTree;
 	
-	// The current gesture sort to use.
+	/** The current gesture sort to use. */
 	FGestureSort GestureSortMode;
 	
-	// The current list of filter strings to filter gestures by.
+	/** The current list of filter strings to filter gestures by. */
 	TArray<FString> FilterStrings;
 };
diff --git a/Engine/Source/Editor/InputBindingEditor/Public/InputBindingEditor.h b/Engine/Source/Editor/InputBindingEditor/Public/InputBindingEditor.h
--- a/Engine/Source/Editor/InputBindingEditor/Public/InputBindingEditor.h
+++ b/Engine/Source/Editor/InputBindingEditor/Public/InputBindingEditor.h
@@ -1,9 +1,5 @@
 // Copyright 1998-2014 Epic Games, Inc. All Rights Reserved.
 
-/*=============================================================================
-	InputBindingEditor.h: InputBindingEditor module public header file.
-=============================================================================*/
-
 #pragma once
 
 
diff --git a/Engine/Source/Editor/InputBindingEditor/Public/Interfaces/IInputBindingEditorModule.h b/Engine/Source/Editor/InputBindingEditor/Public/Interfaces/IInputBindingEditorModule.h
--- a/Engine/Source/Editor/InputBindingEditor/Public/Interfaces/IInputBindingEditorModule.h
+++ b/Engine/Source/Editor/InputBindingEditor/Public/Interfaces/IInputBindingEditorModule.h
@@ -1,9 +1,5 @@
 // Copyright 1998-2014 Epic Games, Inc. All Rights Reserved.
 
-/*=============================================================================
-	IInputBindingEditorModule.h: Declares the IInputBindingEditorModule interface.
-=============================================================================*/
-
 #pragma once
 
 
@@ -19,21 +15,20 @@ public:
 	 * Creates an input binding editor panel widget.
 	 *
 	 * @return The created widget.
+	 * @see DestroyInputBindingEditorPanel
 	 */
 	virtual TWeakPtr<SWidget> CreateInputBindingEditorPanel( ) = 0;
 
 	/**
 	 * Destroys a previously created editor panel widget.
 	 *
-	 * @param Panel - The panel to destroy.
+	 * @param Panel The panel to destroy.
+	 * CreateInputBindingEditorPanel
 	 */
 	virtual void DestroyInputBindingEditorPanel( const TWeakPtr<SWidget>& Panel ) = 0;
 
-
 public:
 
-	/**
-	 * Virtual destructor.
-	 */
+	/** Virtual destructor. */
 	virtual ~IInputBindingEditorModule( ) { }
 };
diff --git a/Engine/Source/Editor/Kismet/Private/Debugging/SKismetDebuggingView.h b/Engine/Source/Editor/Kismet/Private/Debugging/SKismetDebuggingView.h
--- a/Engine/Source/Editor/Kismet/Private/Debugging/SKismetDebuggingView.h
+++ b/Engine/Source/Editor/Kismet/Private/Debugging/SKismetDebuggingView.h
@@ -1,9 +1,9 @@
 // Copyright 1998-2014 Epic Games, Inc. All Rights Reserved.
 
-
 #pragma once
 
-#include "../../Public/BlueprintEditor.h"
+#include "Public/BlueprintEditor.h"
+
 
 //////////////////////////////////////////////////////////////////////////
 // FDebugLineItem
diff --git a/Engine/Source/Editor/SwarmInterface/Private/SwarmInterfaceModule.cpp b/Engine/Source/Editor/SwarmInterface/Private/SwarmInterfaceModule.cpp
--- a/Engine/Source/Editor/SwarmInterface/Private/SwarmInterfaceModule.cpp
+++ b/Engine/Source/Editor/SwarmInterface/Private/SwarmInterfaceModule.cpp
@@ -1,24 +1,20 @@
 // Copyright 1998-2014 Epic Games, Inc. All Rights Reserved.
 
 #include "CoreUObject.h"
-
 #include "ModuleManager.h"
-
 #include "SwarmMessages.h"
 
-class FSwarmInterfaceModule : public IModuleInterface
-{
 
+class FSwarmInterfaceModule
+	: public IModuleInterface
+{
 public:
 
-	virtual void StartupModule() override {}
+	// IModuleInterface interface
 
-	virtual void ShutdownModule() override {}
+	virtual void StartupModule() override { }
+	virtual void ShutdownModule() override { }
 };
 
-// Dummy class initialization
-USwarmMessages::USwarmMessages( const class FPostConstructInitializeProperties& PCIP )
-	: Super(PCIP)
-{ }
 
-IMPLEMENT_MODULE( FSwarmInterfaceModule, SwarmInterface );
+IMPLEMENT_MODULE(FSwarmInterfaceModule, SwarmInterface);
diff --git a/Engine/Source/Editor/TaskBrowser/Private/STaskBrowser.cpp b/Engine/Source/Editor/TaskBrowser/Private/STaskBrowser.cpp
--- a/Engine/Source/Editor/TaskBrowser/Private/STaskBrowser.cpp
+++ b/Engine/Source/Editor/TaskBrowser/Private/STaskBrowser.cpp
@@ -5,8 +5,8 @@
 #include "STaskColumn.h"
 #include "STaskComplete.h"
 #include "STaskSettings.h"
-#include "../Public/TaskDatabase.h"
-#include "../Public/TaskDataManager.h"
+#include "Public/TaskDatabase.h"
+#include "Public/TaskDataManager.h"
 
 #if PLATFORM_WINDOWS
 #include "AllowWindowsPlatformTypes.h"
@@ -33,7 +33,6 @@
 DEFINE_LOG_CATEGORY_STATIC( LogTaskBrowser, Log, All );
 
 
-
 /**
  * Reads the mac address for the computer
  *
diff --git a/Engine/Source/Editor/UnrealEdMessages/Classes/AssetEditorMessages.h b/Engine/Source/Editor/UnrealEdMessages/Classes/AssetEditorMessages.h
--- a/Engine/Source/Editor/UnrealEdMessages/Classes/AssetEditorMessages.h
+++ b/Engine/Source/Editor/UnrealEdMessages/Classes/AssetEditorMessages.h
@@ -1,9 +1,5 @@
 // Copyright 1998-2014 Epic Games, Inc. All Rights Reserved.
 
-/*=============================================================================
-	AssetEditorMessages.h: Declares message types used by asset editor.
-=============================================================================*/
-
 #pragma once
 
 #include "AssetEditorMessages.generated.h"
@@ -46,16 +42,3 @@ struct TStructOpsTypeTraits<FAssetEditorRequestOpenAsset> : public TStructOpsTyp
 		WithMessageHandling = true
 	};
 };
-
-
-/* Dummy class
- *****************************************************************************/
-
-UCLASS()
-class UAssetEditorMessages
-	: public UObject
-{
-public:
-
-	GENERATED_UCLASS_BODY()
-};
diff --git a/Engine/Source/Editor/UnrealEdMessages/Classes/FileServerMessages.h b/Engine/Source/Editor/UnrealEdMessages/Classes/FileServerMessages.h
--- a/Engine/Source/Editor/UnrealEdMessages/Classes/FileServerMessages.h
+++ b/Engine/Source/Editor/UnrealEdMessages/Classes/FileServerMessages.h
@@ -1,9 +1,5 @@
 // Copyright 1998-2014 Epic Games, Inc. All Rights Reserved.
 
-/*=============================================================================
-	FileServerMessages.h: Declares message types used by file server commandlets.
-=============================================================================*/
-
 #pragma once
 
 #include "FileServerMessages.generated.h"
diff --git a/Engine/Source/Editor/UnrealEdMessages/Private/UnrealEdMessagesModule.cpp b/Engine/Source/Editor/UnrealEdMessages/Private/UnrealEdMessagesModule.cpp
--- a/Engine/Source/Editor/UnrealEdMessages/Private/UnrealEdMessagesModule.cpp
+++ b/Engine/Source/Editor/UnrealEdMessages/Private/UnrealEdMessagesModule.cpp
@@ -1,13 +1,8 @@
 // Copyright 1998-2014 Epic Games, Inc. All Rights Reserved.
 
-/*=============================================================================
-	UnrealEdMessagesModule.cpp: Implements the FUnrealEdMessagesModule class.
-=============================================================================*/
-
 #include "UnrealEdMessagesPrivatePCH.h"
 
 
-
 /**
  * Implements the UnrealEdMessages module.
  */
@@ -16,8 +11,9 @@ class FUnrealEdMessagesModule
 {
 public:
 
-	virtual void StartupModule( ) override { }
+	// IModuleInterface interface
 
+	virtual void StartupModule( ) override { }
 	virtual void ShutdownModule( ) override { }
 
 	virtual bool SupportsDynamicReloading( ) override
@@ -27,10 +23,4 @@ public:
 };
 
 
-// Dummy class initialization
-UAssetEditorMessages::UAssetEditorMessages( const class FPostConstructInitializeProperties& PCIP )
-	: Super(PCIP)
-{ }
-
-
 IMPLEMENT_MODULE(FUnrealEdMessagesModule, UnrealEdMessages);
diff --git a/Engine/Source/Editor/UnrealEdMessages/Private/UnrealEdMessagesPrivatePCH.h b/Engine/Source/Editor/UnrealEdMessages/Private/UnrealEdMessagesPrivatePCH.h
--- a/Engine/Source/Editor/UnrealEdMessages/Private/UnrealEdMessagesPrivatePCH.h
+++ b/Engine/Source/Editor/UnrealEdMessages/Private/UnrealEdMessagesPrivatePCH.h
@@ -1,24 +1,11 @@
 // Copyright 1998-2014 Epic Games, Inc. All Rights Reserved.
 
-/*=============================================================================
-	UnrealEdMessagesPrivatePCH.h: Pre-compiled header file for the UnrealEdMessages module.
-=============================================================================*/
+#pragma once
 
-#ifndef UNREALEDMESSAGES_PRIVATEPCH_H
-#define UNREALEDMESSAGES_PRIVATEPCH_H
-
-
-#include "../Public/UnrealEdMessages.h"
+#include "Public/UnrealEdMessages.h"
 
 
 /* Dependencies
  *****************************************************************************/
 
 #include "ModuleManager.h"
-
-
-/* Private includes
- *****************************************************************************/
-
-
-#endif
diff --git a/Engine/Source/Editor/UnrealEdMessages/Public/UnrealEdMessages.h b/Engine/Source/Editor/UnrealEdMessages/Public/UnrealEdMessages.h
--- a/Engine/Source/Editor/UnrealEdMessages/Public/UnrealEdMessages.h
+++ b/Engine/Source/Editor/UnrealEdMessages/Public/UnrealEdMessages.h
@@ -1,9 +1,5 @@
 // Copyright 1998-2014 Epic Games, Inc. All Rights Reserved.
 
-/*=============================================================================
-	UnrealEdMessages.h: UnrealEdMessages module public header file.
-=============================================================================*/
-
 #pragma once
 
 
diff --git a/Engine/Source/Programs/UnrealFrontend/Private/Linux/LinuxUnrealFrontendMain.cpp b/Engine/Source/Programs/UnrealFrontend/Private/Linux/LinuxUnrealFrontendMain.cpp
--- a/Engine/Source/Programs/UnrealFrontend/Private/Linux/LinuxUnrealFrontendMain.cpp
+++ b/Engine/Source/Programs/UnrealFrontend/Private/Linux/LinuxUnrealFrontendMain.cpp
@@ -1,13 +1,11 @@
 // Copyright 1998-2014 Epic Games, Inc. All Rights Reserved.
 
-/*=============================================================================
-	LinuxUnrealFrontendMain.cpp: Implements the main entry point for Linux.
-=============================================================================*/
-
 #include "UnrealFrontendMain.h"
 
+
 static FString GSavedCommandLine;
 
+
 int main( int argc, char *argv[] )
 {
 	for (int32 Option = 1; Option < argc; ++Option)
diff --git a/Engine/Source/Runtime/Core/Private/GenericPlatform/GenericPlatformProcess.cpp b/Engine/Source/Runtime/Core/Private/GenericPlatform/GenericPlatformProcess.cpp
--- a/Engine/Source/Runtime/Core/Private/GenericPlatform/GenericPlatformProcess.cpp
+++ b/Engine/Source/Runtime/Core/Private/GenericPlatform/GenericPlatformProcess.cpp
@@ -1,7 +1,7 @@
 // Copyright 1998-2014 Epic Games, Inc. All Rights Reserved.
 
 #include "Core.h"
-#include "../../Public/Modules/ModuleVersion.h"
+#include "Public/Modules/ModuleVersion.h"
 
 
 #if PLATFORM_HAS_BSD_TIME 
diff --git a/Engine/Source/Runtime/Core/Private/Linux/LinuxPlatformProcess.cpp b/Engine/Source/Runtime/Core/Private/Linux/LinuxPlatformProcess.cpp
--- a/Engine/Source/Runtime/Core/Private/Linux/LinuxPlatformProcess.cpp
+++ b/Engine/Source/Runtime/Core/Private/Linux/LinuxPlatformProcess.cpp
@@ -1,12 +1,8 @@
 // Copyright 1998-2014 Epic Games, Inc. All Rights Reserved.
 
-/*=============================================================================
-	LinuxPlatformProcess.cpp: Linux implementations of Process functions
-=============================================================================*/
-
 #include "Core.h"
 #include "LinuxPlatformRunnableThread.h"
-#include "../../Public/Modules/ModuleVersion.h"
+#include "Public/Modules/ModuleVersion.h"
 #include <spawn.h>
 #include <sys/wait.h>
 #include <sys/resource.h>
diff --git a/Engine/Source/Runtime/Core/Resources/Windows/ModuleVersionResource.rc.inl b/Engine/Source/Runtime/Core/Resources/Windows/ModuleVersionResource.rc.inl
--- a/Engine/Source/Runtime/Core/Resources/Windows/ModuleVersionResource.rc.inl
+++ b/Engine/Source/Runtime/Core/Resources/Windows/ModuleVersionResource.rc.inl
@@ -1,17 +1,18 @@
 // Copyright 1998-2014 Epic Games, Inc. All Rights Reserved.
 
 #include <windows.h>
-#include "../../Public/Modules/ModuleVersion.h"
+#include "Public/Modules/ModuleVersion.h"
 #include "ModuleVersionResource.h"
 #include "CoreMiscDefines.h"
 
-// Netral language
+
+// Neutral language
 LANGUAGE LANG_NEUTRAL,SUBLANG_NEUTRAL
 
 // Declare the engine API version string
 #if !IS_MONOLITHIC
-ID_MODULE_API_VERSION_RESOURCE RCDATA
-{
-	PREPROCESSOR_TO_STRING(MODULE_API_VERSION) 0
-}
+	ID_MODULE_API_VERSION_RESOURCE RCDATA
+	{
+		PREPROCESSOR_TO_STRING(MODULE_API_VERSION) 0
+	}
 #endif
\ No newline at end of file
diff --git a/Engine/Source/Runtime/Engine/Private/AI/Navigation/NavGraphGenerator.h b/Engine/Source/Runtime/Engine/Private/AI/Navigation/NavGraphGenerator.h
--- a/Engine/Source/Runtime/Engine/Private/AI/Navigation/NavGraphGenerator.h
+++ b/Engine/Source/Runtime/Engine/Private/AI/Navigation/NavGraphGenerator.h
@@ -2,7 +2,7 @@
 
 #pragma once 
 
-#include "../../Public/AI/NavDataGenerator.h"
+#include "AI/NavDataGenerator.h"
 
 /**
  * Class that handles generation of the ANavigationGraph data
diff --git a/Engine/Source/Runtime/IOS/LaunchDaemonMessages/Private/LaunchDaemonMessagesModule.cpp b/Engine/Source/Runtime/IOS/LaunchDaemonMessages/Private/LaunchDaemonMessagesModule.cpp
--- a/Engine/Source/Runtime/IOS/LaunchDaemonMessages/Private/LaunchDaemonMessagesModule.cpp
+++ b/Engine/Source/Runtime/IOS/LaunchDaemonMessages/Private/LaunchDaemonMessagesModule.cpp
@@ -1,13 +1,8 @@
 // Copyright 1998-2014 Epic Games, Inc. All Rights Reserved.
 
-/*=============================================================================
-	EngineMessagesModule.cpp: Implements the FEngineMessagesModule class.
-=============================================================================*/
-
 #include "LaunchDaemonMessagesPrivatePCH.h"
 
 
-
 /**
  * Implements the EngineMessages module.
  */
@@ -26,9 +21,5 @@ public:
 	}
 };
 
-// Dummy class initialization
-UIOSMessageProtocol::UIOSMessageProtocol( const class FPostConstructInitializeProperties& PCIP )
-	: Super(PCIP)
-{ }
 
 IMPLEMENT_MODULE(FLaunchDaemonMessagesModule, LaunchDaemonMessages);
diff --git a/Engine/Source/Runtime/IOS/LaunchDaemonMessages/Private/LaunchDaemonMessagesPrivatePCH.h b/Engine/Source/Runtime/IOS/LaunchDaemonMessages/Private/LaunchDaemonMessagesPrivatePCH.h
--- a/Engine/Source/Runtime/IOS/LaunchDaemonMessages/Private/LaunchDaemonMessagesPrivatePCH.h
+++ b/Engine/Source/Runtime/IOS/LaunchDaemonMessages/Private/LaunchDaemonMessagesPrivatePCH.h
@@ -1,24 +1,11 @@
 // Copyright 1998-2014 Epic Games, Inc. All Rights Reserved.
 
-/*=============================================================================
-	EngineMessagesPrivatePCH.h: Pre-compiled header file for the EngineMessages module.
-=============================================================================*/
+#pragma once
 
-#ifndef LAUNCHDAEMONMESSAGES_PRIVATEPCH_H
-#define LAUNCHDAEMONMESSAGES_PRIVATEPCH_H
-
-
-#include "../Public/LaunchDaemonMessages.h"
+#include "LaunchDaemonMessages.h"
 
 
 /* Dependencies
  *****************************************************************************/
 
 #include "ModuleManager.h"
-
-
-/* Private includes
- *****************************************************************************/
-
-
-#endif
diff --git a/Engine/Source/Runtime/IOS/LaunchDaemonMessages/Public/LaunchDaemonMessages.h b/Engine/Source/Runtime/IOS/LaunchDaemonMessages/Public/LaunchDaemonMessages.h
--- a/Engine/Source/Runtime/IOS/LaunchDaemonMessages/Public/LaunchDaemonMessages.h
+++ b/Engine/Source/Runtime/IOS/LaunchDaemonMessages/Public/LaunchDaemonMessages.h
@@ -1,9 +1,5 @@
 // Copyright 1998-2014 Epic Games, Inc. All Rights Reserved.
 
-/*=============================================================================
-	LaunchDaemonMessages.h: LaunchDaemonMessages module public header file.
-=============================================================================*/
-
 #pragma once
 
 
diff --git a/Engine/Source/Runtime/Launch/Private/Launch.cpp b/Engine/Source/Runtime/Launch/Private/Launch.cpp
--- a/Engine/Source/Runtime/Launch/Private/Launch.cpp
+++ b/Engine/Source/Runtime/Launch/Private/Launch.cpp
@@ -1,14 +1,11 @@
 // Copyright 1998-2014 Epic Games, Inc. All Rights Reserved.
 
-/*=============================================================================
-	Launch.cpp: Game launcher.
-=============================================================================*/
-
 #include "LaunchPrivatePCH.h"
 #include "PhysicsPublic.h"
 #include "ExceptionHandling.h"
 #include "ModuleManager.h"
 
+
 IMPLEMENT_MODULE(FDefaultModuleImpl, Launch);
 
 #if UE_EDITOR && IS_MONOLITHIC
@@ -20,6 +17,7 @@ PER_MODULE_BOILERPLATE
 FEngineLoop	GEngineLoop;
 bool GIsConsoleExecutable = false;
 
+
 extern "C" int test_main(int argc, char ** argp)
 {
 	return 0;
diff --git a/Engine/Source/Runtime/Launch/Private/LaunchEngineLoop.cpp b/Engine/Source/Runtime/Launch/Private/LaunchEngineLoop.cpp
--- a/Engine/Source/Runtime/Launch/Private/LaunchEngineLoop.cpp
+++ b/Engine/Source/Runtime/Launch/Private/LaunchEngineLoop.cpp
@@ -1,9 +1,5 @@
 // Copyright 1998-2014 Epic Games, Inc. All Rights Reserved.
 
-/*=============================================================================
-	LaunchEngineLoop.cpp: Implements the main engine loop.
-=============================================================================*/
-
 #include "LaunchPrivatePCH.h"
 #include "Internationalization/Internationalization.h"
 #include "Ticker.h"
diff --git a/Engine/Source/Runtime/Launch/Private/LaunchPrivatePCH.h b/Engine/Source/Runtime/Launch/Private/LaunchPrivatePCH.h
--- a/Engine/Source/Runtime/Launch/Private/LaunchPrivatePCH.h
+++ b/Engine/Source/Runtime/Launch/Private/LaunchPrivatePCH.h
@@ -1,9 +1,5 @@
 // Copyright 1998-2014 Epic Games, Inc. All Rights Reserved.
 
-/*=============================================================================
-	LaunchPrivatePCH.h: Pre-compiled header file for the Launch module.
-=============================================================================*/
-
 #pragma once
 
 /* Dependencies
diff --git a/Engine/Source/Runtime/Launch/Public/LaunchEngineLoop.h b/Engine/Source/Runtime/Launch/Public/LaunchEngineLoop.h
--- a/Engine/Source/Runtime/Launch/Public/LaunchEngineLoop.h
+++ b/Engine/Source/Runtime/Launch/Public/LaunchEngineLoop.h
@@ -1,16 +1,14 @@
 // Copyright 1998-2014 Epic Games, Inc. All Rights Reserved.
 
-/*=============================================================================
-	LaunchEngineLoop.h: Declares the FEngineLoop class.
-=============================================================================*/
-
 #pragma once
 
 #if WITH_ENGINE
-#include "ISessionService.h"
+	#include "ISessionService.h"
 #endif
+
 class FEngineService;
 
+
 /**
  * Implements the main engine loop.	
  */
@@ -29,70 +27,54 @@ public:
 public:
 
 	/**
-	 * Pre-Initialize the main loop, and generates the commandline from standard ArgC/ArgV from main()
+	 * Pre-Initialize the main loop, and generates the commandline from standard ArgC/ArgV from main().
 	 *
-	 * @param ArgC - The number of strings in ArgV
-	 * @param ArgV - The command line parameters (ArgV[0] is expected to be the executable name)
-	 * @param AdditionalCommandLine - Optional string to append to the command line (after ArgV is put together)
-	 *
-	 * @return Returns the error level, 0 if successful and > 0 if there were errors
+	 * @param ArgC The number of strings in ArgV.
+	 * @param ArgV The command line parameters (ArgV[0] is expected to be the executable name).
+	 * @param AdditionalCommandLine Optional string to append to the command line (after ArgV is put together).
+	 * @return Returns the error level, 0 if successful and > 0 if there were errors.
 	 */ 
 	int32 PreInit(int32 ArgC, TCHAR* ArgV[], const TCHAR* AdditionalCommandline = NULL);
 
 	/**
 	 * Pre-Initialize the main loop - parse command line, sets up GIsEditor, etc.
 	 *
-	 * @param CmdLine - The command line.
-	 *
-	 * @return The error level; 0 if successful, > 0 if there were errors
+	 * @param CmdLine The command line.
+	 * @return The error level; 0 if successful, > 0 if there were errors.
 	 */ 
 	int32 PreInit(const TCHAR* CmdLine);
 
-	/**
-	 * Load all modules needed before Init.
-	 */ 
+	/** Load all modules needed before Init. */ 
 	void LoadPreInitModules();
 
-	/**
-	* Load core modules.
-	*/
+	/** Load core modules. */
 	void LoadCoreModules();
 
 #if WITH_ENGINE
 	
-	/**
-	 * Load all core modules needed at startup time.
-	 */
+	/** Load all core modules needed at startup time. */
 	bool LoadStartupCoreModules();
 	
-	/**
-	 * Load all modules needed at startup time.
-	 */
+	/** Load all modules needed at startup time. */
 	bool LoadStartupModules();
 
 	/**
-	 * Initialize the main loop - the rest of the initialization.
+	 * Initialize the main loop (the rest of the initialization).
 	 *
 	 * @return The error level; 0 if successful, > 0 if there were errors.
 	 */ 
 	virtual int32 Init() override;
 
-	/**
-	 * Initialize the timing options from the command line.
-	 */ 
+	/** Initialize the timing options from the command line. */ 
 	void InitTime();
 
-	/** 
-	 * Performs shut down.
-	 */
+	/** Performs shut down. */
 	void Exit();
 
 	/** Whether the engine should operate in an idle mode that uses no CPU or GPU time. */
 	bool ShouldUseIdleMode() const;
 
-	/**
-	 * Advances the main loop.
-	 */
+	/** Advances the main loop. */
 	virtual void Tick() override;
 
 	/** Removes references to any objects pending cleanup by deleting them. */
@@ -106,14 +88,12 @@ public:
 
 #endif // WITH_ENGINE
 
-	/** Pre-init HMD device (if necessary) */
+	/** Pre-init HMD device (if necessary). */
 	static void PreInitHMDDevice();
 
 public:
 
-	/**
-	 * Initializes the application.
-	 */
+	/** Initializes the application. */
 	static bool AppInit( );
 
 	/**
@@ -131,43 +111,30 @@ public:
 	static void AppExit( );
 
 private:
-	/**
-	 * Utility function that processes Slate operations
-	 */
+
+	/** Utility function that processes Slate operations. */
 	void ProcessPlayerControllersSlateOperations() const;
 
 protected:
 
-	/**
-	 * Holds a dynamically expanding array of frame times in milliseconds (if FApp::IsBenchmarking() is set).
-	 */
+	/** Holds a dynamically expanding array of frame times in milliseconds (if FApp::IsBenchmarking() is set). */
 	TArray<float> FrameTimes;
 
-	/**
-	 * Holds the total time spent ticking engine.
-	 */
+	/** Holds the total time spent ticking engine. */
 	double TotalTickTime;
 	
-	/**
-	 * Holds the maximum number of seconds engine should be ticked.
-	 */
+	/** Holds the maximum number of seconds engine should be ticked. */
 	double MaxTickTime;
 	
-	/**
-	 * Holds the maximum number of frames to render in benchmarking mode.
-	 */
+	/** Holds the maximum number of frames to render in benchmarking mode. */
 	uint64 MaxFrameCounter;
 	
-	/**
-	 * Holds the number of cycles in the last frame.
-	 */
+	/** Holds the number of cycles in the last frame. */
 	uint32 LastFrameCycles;
 
 #if WITH_ENGINE
 
-	/**
-	 * Holds the objects which need to be cleaned up when the rendering thread finishes the previous frame.
-	 */
+	/** Holds the objects which need to be cleaned up when the rendering thread finishes the previous frame. */
 	FPendingCleanupObjects* PendingCleanupObjects;
 
 #endif //WITH_ENGINE
@@ -176,10 +143,10 @@ private:
 
 #if WITH_ENGINE
 
-	// Holds the engine service.
+	/** Holds the engine service. */
 	FEngineService* EngineService;
 
-	// Holds the application session service.
+	/** Holds the application session service. */
 	ISessionServicePtr SessionService;
 
 #endif // WITH_ENGINE
diff --git a/Engine/Source/Runtime/Launch/Public/RequiredProgramMainCPPInclude.h b/Engine/Source/Runtime/Launch/Public/RequiredProgramMainCPPInclude.h
--- a/Engine/Source/Runtime/Launch/Public/RequiredProgramMainCPPInclude.h
+++ b/Engine/Source/Runtime/Launch/Public/RequiredProgramMainCPPInclude.h
@@ -1,10 +1,5 @@
 // Copyright 1998-2014 Epic Games, Inc. All Rights Reserved.
 
-/*=============================================================================
-	LaunchEngineLoopForProgram.h: 
-		includes the stuff needed for the main cpp of a program
-=============================================================================*/
-
 #pragma once
 
 #include "ModuleManager.h"
diff --git a/Engine/Source/Runtime/NullDrv/Private/NullDrv.cpp b/Engine/Source/Runtime/NullDrv/Private/NullDrv.cpp
--- a/Engine/Source/Runtime/NullDrv/Private/NullDrv.cpp
+++ b/Engine/Source/Runtime/NullDrv/Private/NullDrv.cpp
@@ -1,18 +1,15 @@
 // Copyright 1998-2014 Epic Games, Inc. All Rights Reserved.
 
-/*=============================================================================
-	NullDrv.cpp: Unreal Null RHI library implementation.
-=============================================================================*/
-
 #include "NullDrvPrivate.h"
 #include "ModuleManager.h"
-
-IMPLEMENT_MODULE(FNullDynamicRHIModule, NullDrv);
-
 #include "RenderCore.h"
 #include "RenderResource.h"
 
+
 bool FNullDynamicRHIModule::IsSupported()
 {
 	return true;
-}
\ No newline at end of file
+}
+
+
+IMPLEMENT_MODULE(FNullDynamicRHIModule, NullDrv);
diff --git a/Engine/Source/Runtime/NullDrv/Private/NullRHI.cpp b/Engine/Source/Runtime/NullDrv/Private/NullRHI.cpp
--- a/Engine/Source/Runtime/NullDrv/Private/NullRHI.cpp
+++ b/Engine/Source/Runtime/NullDrv/Private/NullRHI.cpp
@@ -1,9 +1,5 @@
 // Copyright 1998-2014 Epic Games, Inc. All Rights Reserved.
 
-/*=============================================================================
-	NullRHI.cpp: Null RHI implementation.
-=============================================================================*/
-
 #include "NullDrvPrivate.h"
 #include "RHI.h"
 
@@ -12,11 +8,13 @@
 #include "NullRHI.h"
 #include "RenderResource.h"
 
+
 FNullDynamicRHI::FNullDynamicRHI()
 {
 	GMaxRHIShaderPlatformValue = ShaderFormatToLegacyShaderPlatform(FName(FPlatformMisc::GetNullRHIShaderFormat()));
 }
 
+
 void FNullDynamicRHI::Init()
 {
 #if PLATFORM_WINDOWS
@@ -47,10 +45,12 @@ void FNullDynamicRHI::Init()
 	GIsRHIInitialized = true;
 }
 
+
 void FNullDynamicRHI::Shutdown()
 {
 }
 
+
 /**
  * Return a shared large static buffer that can be used to return from any 
  * function that needs to return a valid pointer (but can be garbage data)
@@ -67,13 +67,16 @@ void* FNullDynamicRHI::GetStaticBuffer()
 	return Buffer;
 }
 
+
 /** Value between 0-100 that determines the percentage of the vertical scan that is allowed to pass while still allowing us to swap when VSYNC'ed.
 This is used to get the same behavior as the old *_OR_IMMEDIATE present modes. */
 uint32 GPresentImmediateThreshold = 100;
 
-#else
+
+#else //USE_DYNAMIC_RHI
+
 
 // Suppress linker warning "warning LNK4221: no public symbols found; archive member will be inaccessible"
 int32 NullRHILinkerHelper;
 
-#endif // USE_DYNAMIC_RHI
+#endif //USE_DYNAMIC_RHI
diff --git a/Engine/Source/Runtime/NullDrv/Public/NullDrv.h b/Engine/Source/Runtime/NullDrv/Public/NullDrv.h
--- a/Engine/Source/Runtime/NullDrv/Public/NullDrv.h
+++ b/Engine/Source/Runtime/NullDrv/Public/NullDrv.h
@@ -1,9 +1,5 @@
 // Copyright 1998-2014 Epic Games, Inc. All Rights Reserved.
 
-/*=============================================================================
-	NullDrv.h: Public Null RHI definitions.
-=============================================================================*/
-
 #pragma once
 
 // Dependencies.
@@ -11,19 +7,20 @@
 #include "RHI.h"
 #include "NullRHI.h"
 
+
 /** Implements the NullDrv module as a dynamic RHI providing module. */
-class FNullDynamicRHIModule : public IDynamicRHIModule
+class FNullDynamicRHIModule
+	: public IDynamicRHIModule
 {
 public:
 
-	// IModuleInterface
-	virtual bool SupportsDynamicReloading() override { return false; }
-
 	// IDynamicRHIModule
+
+	virtual bool SupportsDynamicReloading() override { return false; }
 	virtual bool IsSupported() override;
 
 	virtual FDynamicRHI* CreateRHI() override
 	{
 		return new FNullDynamicRHI();
 	}
-};
\ No newline at end of file
+};
diff --git a/Engine/Source/Runtime/NullDrv/Public/NullRHI.h b/Engine/Source/Runtime/NullDrv/Public/NullRHI.h
--- a/Engine/Source/Runtime/NullDrv/Public/NullRHI.h
+++ b/Engine/Source/Runtime/NullDrv/Public/NullRHI.h
@@ -1,14 +1,11 @@
 // Copyright 1998-2014 Epic Games, Inc. All Rights Reserved.
 
-/*=============================================================================	
-	NullRHI.h: Null RHI definitions.
-=============================================================================*/
+#pragma once
 
-#ifndef __NULLRHI_H__
-#define __NULLRHI_H__
 
 /** A null implementation of the dynamically bound RHI. */
-class FNullDynamicRHI : public FDynamicRHI
+class FNullDynamicRHI
+	: public FDynamicRHI
 {
 public:
 
@@ -29,5 +26,3 @@ private:
 	/** Allocates a static buffer for RHI functions to return as a write destination. */
 	static void* GetStaticBuffer();
 };
-
-#endif
diff --git a/Engine/Source/Runtime/Slate/Public/Framework/Application/IWidgetReflector.h b/Engine/Source/Runtime/Slate/Public/Framework/Application/IWidgetReflector.h
--- a/Engine/Source/Runtime/Slate/Public/Framework/Application/IWidgetReflector.h
+++ b/Engine/Source/Runtime/Slate/Public/Framework/Application/IWidgetReflector.h
@@ -1,9 +1,5 @@
 // Copyright 1998-2014 Epic Games, Inc. All Rights Reserved.
 
-/*=============================================================================
-	IWidgetReflector.h: Declares the IWidgetReflector class.
-=============================================================================*/
-
 #pragma once
 
 
@@ -18,9 +14,7 @@ class IWidgetReflector
 {
 public:
 
-	/**
-	 * Called when the user has picked a widget to observe.
-	 */
+	/** Called when the user has picked a widget to observe. */
 	virtual void OnWidgetPicked( ) = 0;
 
 	/**
@@ -77,8 +71,6 @@ public:
 
 public:
 
-	/**
-	 * Virtual destructor.
-	 */
+	/** Virtual destructor. */
 	virtual ~IWidgetReflector( ) { }
 };
diff --git a/Engine/Source/Runtime/Sockets/Private/BSDSockets/SocketsBSD.cpp b/Engine/Source/Runtime/Sockets/Private/BSDSockets/SocketsBSD.cpp
--- a/Engine/Source/Runtime/Sockets/Private/BSDSockets/SocketsBSD.cpp
+++ b/Engine/Source/Runtime/Sockets/Private/BSDSockets/SocketsBSD.cpp
@@ -1,9 +1,5 @@
 // Copyright 1998-2014 Epic Games, Inc. All Rights Reserved.
 
-/*=============================================================================
-	SocketsBSD.cpp: Implements the FSocketBSD class.
-=============================================================================*/
-
 #include "SocketsPrivatePCH.h"
 
 #if PLATFORM_HAS_BSD_SOCKETS
diff --git a/Engine/Source/Runtime/Sockets/Private/BSDSockets/SocketsBSD.h b/Engine/Source/Runtime/Sockets/Private/BSDSockets/SocketsBSD.h
--- a/Engine/Source/Runtime/Sockets/Private/BSDSockets/SocketsBSD.h
+++ b/Engine/Source/Runtime/Sockets/Private/BSDSockets/SocketsBSD.h
@@ -1,9 +1,5 @@
 // Copyright 1998-2014 Epic Games, Inc. All Rights Reserved.
 
-/*=============================================================================
-	SocketsBSD.h: Declares the FSocketBSD class.
-=============================================================================*/
-
 #pragma once
 
 #include "Platform.h"
@@ -13,6 +9,7 @@
 #include "Sockets.h"
 #include "IPAddress.h"
 
+
 namespace ESocketInternalState
 {
 	enum Param
@@ -29,6 +26,7 @@ namespace ESocketInternalState
 	};
 };
 
+
 /**
  * Implements a BSD network socket.
  */
@@ -38,11 +36,11 @@ class FSocketBSD
 public:
 
 	/**
-	 * Assigns a BSD socket to this object
+	 * Assigns a BSD socket to this object.
 	 *
-	 * @param InSocket the socket to assign to this object
-	 * @param InSocketType the type of socket that was created
-	 * @param InSocketDescription the debug description of the socket
+	 * @param InSocket the socket to assign to this object.
+	 * @param InSocketType the type of socket that was created.
+	 * @param InSocketDescription the debug description of the socket.
 	 */
 	FSocketBSD(SOCKET InSocket, ESocketType InSocketType, const FString& InSocketDescription, ISocketSubsystem * InSubsystem) 
 		: FSocket(InSocketType, InSocketDescription)
@@ -60,61 +58,36 @@ public:
 		Close();
 	}
 
-
 public:
 
-	virtual bool Close() override;
+	// FSocket overrides
 
+	virtual bool Close() override;
 	virtual bool Bind(const FInternetAddr& Addr) override;
-
 	virtual bool Connect(const FInternetAddr& Addr) override;
-
 	virtual bool Listen(int32 MaxBacklog) override;
-
 	virtual bool HasPendingConnection(bool& bHasPendingConnection) override;
-
 	virtual bool HasPendingData(uint32& PendingDataSize) override;
-
 	virtual class FSocket* Accept(const FString& SocketDescription) override;
-
 	virtual class FSocket* Accept(FInternetAddr& OutAddr, const FString& SocketDescription) override;
-
 	virtual bool SendTo(const uint8* Data, int32 Count, int32& BytesSent, const FInternetAddr& Destination) override;
-
 	virtual bool Send(const uint8* Data, int32 Count, int32& BytesSent) override;
-
 	virtual bool RecvFrom(uint8* Data, int32 BufferSize, int32& BytesRead, FInternetAddr& Source, ESocketReceiveFlags::Type Flags = ESocketReceiveFlags::None) override;
-
 	virtual bool Recv(uint8* Data,int32 BufferSize,int32& BytesRead, ESocketReceiveFlags::Type Flags = ESocketReceiveFlags::None) override;
-
 	virtual bool Wait(ESocketWaitConditions::Type Condition, FTimespan WaitTime) override;
-
 	virtual ESocketConnectionState GetConnectionState() override;
-
 	virtual void GetAddress(FInternetAddr& OutAddr) override;
-
 	virtual bool SetNonBlocking(bool bIsNonBlocking = true) override;
-
 	virtual bool SetBroadcast(bool bAllowBroadcast = true) override;
-
 	virtual bool JoinMulticastGroup(const FInternetAddr& GroupAddress) override;
-
 	virtual bool LeaveMulticastGroup(const FInternetAddr& GroupAddress) override;
-
 	virtual bool SetMulticastLoopback(bool bLoopback) override;
-
 	virtual bool SetMulticastTtl(uint8 TimeToLive) override;
-
 	virtual bool SetReuseAddr(bool bAllowReuse = true) override;
-
 	virtual bool SetLinger(bool bShouldLinger = true, int32 Timeout = 0) override;
-
 	virtual bool SetRecvErr(bool bUseErrorQueue = true) override;
-
 	virtual bool SetSendBufferSize(int32 Size,int32& NewSize) override;
-
 	virtual bool SetReceiveBufferSize(int32 Size,int32& NewSize) override;
-
 	virtual int32 GetPortNo() override;
 
 	/**
@@ -124,19 +97,21 @@ public:
 	{
 		return Socket;
 	}
+
 protected:
 
-	// This is generally select(), but makes it easier for platforms without select to replace it
+	/** This is generally select(), but makes it easier for platforms without select to replace it. */
 	virtual ESocketInternalState::Return HasState(ESocketInternalState::Param State, FTimespan WaitTime=FTimespan(0));
 
-	// Holds the BSD socket object.
+	/** Holds the BSD socket object. */
 	SOCKET Socket;
 
-	// Last activity time
+	/** Last activity time. */
 	FDateTime LastActivityTime;
 
-	// Pointer to the subsystem that created it
+	/** Pointer to the subsystem that created it. */
 	ISocketSubsystem * SocketSubsystem;
 };
 
+
 #endif	//PLATFORM_HAS_BSD_SOCKETS
diff --git a/Engine/Source/Runtime/Sockets/Public/MultichannelTcpReceiver.h b/Engine/Source/Runtime/Sockets/Public/MultichannelTcpReceiver.h
--- a/Engine/Source/Runtime/Sockets/Public/MultichannelTcpReceiver.h
+++ b/Engine/Source/Runtime/Sockets/Public/MultichannelTcpReceiver.h
@@ -1,9 +1,5 @@
 // Copyright 1998-2014 Epic Games, Inc. All Rights Reserved.
 
-/*=============================================================================
-	MultichannelTcpReceiver.h: Declares the FMultichannelTcpReceiver class.
-=============================================================================*/
-
 #pragma once
 
 
@@ -28,8 +24,8 @@ public:
 	/**
 	 * Creates and initializes a new instance.
 	 *
-	 * @param InSocket - The socket to receive from.
-	 * @param InReceiveDelegate - Delete to handle data as it becomes available, called from the receive thread.
+	 * @param InSocket The socket to receive from.
+	 * @param InReceiveDelegate Delete to handle data as it becomes available, called from the receive thread.
 	 */
 	FMultichannelTcpReceiver( FSocket* InSocket, const FOnMultichannelTcpReceive& InReceiveDelegate )
 		: Socket(InSocket)
@@ -38,9 +34,7 @@ public:
 		Thread = FRunnableThread::Create(this, TEXT("FMultichannelTCPReceiver"), 8 * 1024, TPri_AboveNormal);
 	}
 
-	/**
-	 * Destructor.
-	 */
+	/** Destructor. */
 	~FMultichannelTcpReceiver()
 	{
 		if (Thread)
@@ -50,7 +44,6 @@ public:
 		}
 	}
 
-
 public:
 
 	/**
@@ -61,10 +54,9 @@ public:
 		return BytesReceived;
 	}
 
-
 public:
 
-	// Begin FRunnable interface
+	// FRunnable interface
 
 	virtual bool Init( )
 	{
@@ -111,26 +103,21 @@ public:
 	}
 
 	virtual void Stop( ) { }
-
 	virtual void Exit( ) { }
 
-	// End FRunnable interface
-
-
 private:
 
-	// Holds the number of bytes received so far.
+	/** Holds the number of bytes received so far. */
 	int32 BytesReceived;
 
-	// Holds the socket to use for communication.
+	/** Holds the socket to use for communication. */
 	FSocket* Socket;
 
-	// Holds the thread we are running on.
+	/** Holds the thread we are running on. */
 	FRunnableThread* Thread;
 
-
 private:
 
-	// Holds a delegate to be invoked when data has been received.
+	/** Holds a delegate to be invoked when data has been received. */
 	FOnMultichannelTcpReceive ReceiveDelegate;
-};
\ No newline at end of file
+};
diff --git a/Engine/Source/Runtime/Sockets/Public/MultichannelTcpSender.h b/Engine/Source/Runtime/Sockets/Public/MultichannelTcpSender.h
--- a/Engine/Source/Runtime/Sockets/Public/MultichannelTcpSender.h
+++ b/Engine/Source/Runtime/Sockets/Public/MultichannelTcpSender.h
@@ -1,9 +1,5 @@
 // Copyright 1998-2014 Epic Games, Inc. All Rights Reserved.
 
-/*=============================================================================
-	MultichannelTcpSender.h: Declares the FMultichannelTcpSender class.
-=============================================================================*/
-
 #pragma once
 
 
@@ -24,20 +20,17 @@ class FMultichannelTcpSender
 {
 	enum
 	{
-		/**
-		 * Defines the maximum ayload size per packet (in bytes).
-		 */
+		/** Defines the maximum payload size per packet (in bytes). */
 		MaxPacket = 128 * 1024 - 8
 	};
 
-
 public:
 
 	/**
 	 * Creates and initializes a new instance.
 	 *
-	 * @param InSocket - The socket to send to.
-	 * @param InOkToSendDelegate - Delegate to ask if it is ok to send a packet.
+	 * @param InSocket The socket to send to.
+	 * @param InOkToSendDelegate Delegate to ask if it is ok to send a packet.
 	 */
 	FMultichannelTcpSender( FSocket* InSocket, const FOnMultichannelTcpOkToSend& InOkToSendDelegate )
 		: Socket(InSocket)
@@ -46,9 +39,7 @@ public:
 		Thread = FRunnableThread::Create(this, TEXT("FMultichannelTCPSender"), 8 * 1024, TPri_AboveNormal);
 	}
 
-	/**
-	 * Destructor.
-	 */
+	/** Destructor. */
 	~FMultichannelTcpSender(  )
 	{
 		if (Thread)
@@ -58,12 +49,9 @@ public:
 		}
 	}
 
-
 public:
 
-	/**
-	 * Call when bandwidth tests should be retried, possibly sending data if there is available bandwidth.
-	 **/
+	/** Call when bandwidth tests should be retried, possibly sending data if there is available bandwidth. */
 	void AttemptResumeSending( )
 	{
 		FScopeLock ScopeLock(&SendBuffersCriticalSection);
@@ -83,8 +71,8 @@ public:
 	 *
 	 * This method does not block on bandwidth and never fails.
 	 *
-	 * @param Data - The buffer containing the data to send.
-	 * @param Count - The number of bytes to send from the buffer.
+	 * @param Data The buffer containing the data to send.
+	 * @param Count The number of bytes to send from the buffer.
 	 */
 	void Send( const uint8* Data, int32 Count, uint32 Channel )
 	{
@@ -108,10 +96,9 @@ public:
 		AttemptResumeSendingInternal();
 	}
 
-
 public:
 
-	// Begin FRunnable interface
+	// FRunnable interface
 
 	virtual bool Init( )
 	{
@@ -213,14 +200,9 @@ public:
 		}
 	}
 
-	// End FRunnable interface
-
-
 protected:
 
-	/**
-	 * Internal call similar to AttemptResumeSending, but does not do the requisite lock.
-	 */
+	/** Internal call similar to AttemptResumeSending, but does not do the requisite lock. */
 	void AttemptResumeSendingInternal( )
 	{
 		if (EventToRestart)
@@ -231,30 +213,28 @@ protected:
 		}
 	}
 
-
 private:
 
-	// Holds the total number of payload bytes sent to the socket.
+	/** Holds the total number of payload bytes sent to the socket. */
 	int32 BytesSent;
 
-	// Holds an event to trigger when bandwidth has freed up.
+	/** Holds an event to trigger when bandwidth has freed up. */
 	FScopedEvent* EventToRestart;
 
-	// Holds the buffers to hold pending per-channel data.
+	/** Holds the buffers to hold pending per-channel data. */
 	TMap<uint32, TArray<uint8>*> SendBuffers;
 
-	// Holds the critical section to guard the send buffers.
+	/** Holds the critical section to guard the send buffers. */
 	FCriticalSection SendBuffersCriticalSection;
 
-	// Holds the socket to use.
+	/** Holds the socket to use. */
 	FSocket* Socket;
 
-	// Holds the thread that is running this instance.
+	/** Holds the thread that is running this instance. */
 	FRunnableThread* Thread;
 
-
 private:
 
-	// Holds a delegate to be invoked when checking if there is sufficient bandwidth to send.
+	/** Holds a delegate to be invoked when checking if there is sufficient bandwidth to send. */
 	FOnMultichannelTcpOkToSend OkToSendDelegate;
 };
\ No newline at end of file
diff --git a/Engine/Source/Runtime/Sockets/Public/MultichannelTcpSocket.h b/Engine/Source/Runtime/Sockets/Public/MultichannelTcpSocket.h
--- a/Engine/Source/Runtime/Sockets/Public/MultichannelTcpSocket.h
+++ b/Engine/Source/Runtime/Sockets/Public/MultichannelTcpSocket.h
@@ -1,9 +1,5 @@
 // Copyright 1998-2014 Epic Games, Inc. All Rights Reserved.
 
-/*=============================================================================
-	MultichannelTcpSocket.h: Declares the FMultichannelTcpSocket class.
-=============================================================================*/
-
 #pragma once
 
 
@@ -26,14 +22,13 @@ class FMultichannelTcpSocket
 		ControlChannel = 0
 	};
 
-
 public:
 
 	/**
 	 * Creates and initializes a new instance.
 	 *
-	 * @param InSocket - The underlying socket to use.
-	 * @param InBandwidthLatencyProduct - The maximum amount of unacknowledged data to send.
+	 * @param InSocket The underlying socket to use.
+	 * @param InBandwidthLatencyProduct The maximum amount of unacknowledged data to send.
 	 */
 	FMultichannelTcpSocket( FSocket* InSocket, uint64 InBandwidthLatencyProduct )
 		: BandwidthLatencyProduct(InBandwidthLatencyProduct)
@@ -42,7 +37,6 @@ public:
 		, Socket(InSocket)
 	{ }
 
-
 public:
 
 	/**
@@ -50,9 +44,9 @@ public:
 	 *
 	 * Can be called from any thread, but not multiple threads for one channel at once.
 	 *
-	 * @param Data - The buffer to fill.
-	 * @param Count - The number of bytes to receive.
-	 * @param Channel - The channel to receive on.
+	 * @param Data The buffer to fill.
+	 * @param Count The number of bytes to receive.
+	 * @param Channel The channel to receive on.
 	 */
 	int32 BlockingReceive( uint8* Data, int32 Count, uint32 Channel )
 	{
@@ -111,8 +105,7 @@ public:
 	 *
 	 * Can be called from any thread, but realize that if this returns > 0, another thread could steal the data.
 	 *
-	 * @param Channel - The channel to check.
-	 *
+	 * @param Channel The channel to check.
 	 * @return The number of bytes in the receive buffer.
 	 */
 	int32 DataAvailable( uint32 Channel )
@@ -130,10 +123,9 @@ public:
 	 * Can be called from any thread, but realize that multiple threads hammering
 	 * a channel at once is unlikely to give useful results.
 	 *
-	 * @param Data - The buffer to hold the results, if any.
-	 * @param MaxCount - the number of bytes in the receive buffer.
-	 * @param Channel, - The channel to check.
-	 *
+	 * @param Data The buffer to hold the results, if any.
+	 * @param MaxCount the number of bytes in the receive buffer.
+	 * @param Channel The channel to check.
 	 * @return The number of bytes written into Data.
 	 */
 	int32 PollingReceive( uint8* Data, int32 MaxCount, uint32 Channel )
@@ -173,9 +165,9 @@ public:
 	 * Can be called from any thread, but if you are calling from multiple threads,
 	 * make sure you are sending an atomic unit.
 	 *
-	 * @param Data - The buffer containing the data to send.
-	 * @param Count - The number of bytes to send.
-	 * @param Channel - The channel to send on.
+	 * @param Data The buffer containing the data to send.
+	 * @param Count The number of bytes to send.
+	 * @param Channel The channel to send on.
 	 */
 	void Send( const uint8* Data, int32 Count, uint32 Channel )
 	{
@@ -183,7 +175,6 @@ public:
 		Sender.Send(Data, Count, Channel);
 	}
 
-
 private:
 
 	/** Helper struct for a receive buffer per channel **/
@@ -205,7 +196,6 @@ private:
 		{ }
 	};
 
-
 	// Callback for receiving data from the receiver thread.
 	void HandleReceiverReceive( const TArray<uint8>& Payload, uint32 Channel, bool bForceByteswapping )
 	{
@@ -255,7 +245,7 @@ private:
 		}
 	}
 
-	// Callback for checking if the sender thread is permitted to send a packet.
+	/** Callback for checking if the sender thread is permitted to send a packet. */
 	bool HandleSenderOkToSend( int32 PayloadSize, uint32 Channel )
 	{
 		if (Channel == ControlChannel)
@@ -268,27 +258,26 @@ private:
 		return Result;
 	}
 
-
 private:
 
-	// Holds the maximum amount of unacknowledged data to send.
+	/** Holds the maximum amount of unacknowledged data to send. */
 	uint64 BandwidthLatencyProduct;
 
-	// Holds the buffers for incoming per-channel data.
+	/** Holds the buffers for incoming per-channel data. */
 	TMap<uint32, FReceiveBuffer*> ReceiveBuffers;
 
-	// Holds a critical section to guard the receive buffers.
+	/** Holds a critical section to guard the receive buffers. */
 	FCriticalSection ReceiveBuffersCriticalSection;
 
-	// Holds the receiver thread.
+	/** Holds the receiver thread. */
 	FMultichannelTcpReceiver Receiver;
 
-	// Holds the total number of bytes received by the client (comes from an 'Ack' on the control channel).
+	/** Holds the total number of bytes received by the client (comes from an 'Ack' on the control channel). */
 	int64 RemoteReceiverBytesReceived;
 
-	// Holds the sender thread.
+	/** Holds the sender thread. */
 	FMultichannelTcpSender Sender;
 
-	// Holds the socket used for actual transfers.
+	/** Holds the socket used for actual transfers. */
 	FSocket* Socket;
 };
\ No newline at end of file
diff --git a/Engine/Source/Runtime/StreamingPauseRendering/Private/StreamingPauseRendering.cpp b/Engine/Source/Runtime/StreamingPauseRendering/Private/StreamingPauseRendering.cpp
--- a/Engine/Source/Runtime/StreamingPauseRendering/Private/StreamingPauseRendering.cpp
+++ b/Engine/Source/Runtime/StreamingPauseRendering/Private/StreamingPauseRendering.cpp
@@ -1,28 +1,27 @@
 // Copyright 1998-2014 Epic Games, Inc. All Rights Reserved.
 
-/*=============================================================================
-StreamingPauseRendering.cpp: Streaming pause implementation.
-=============================================================================*/
-
 #include "Engine.h"
 #include "StreamingPauseRendering.h"
 #include "SceneViewport.h"
 #include "MoviePlayer.h"
 
+
 IMPLEMENT_MODULE(FStreamingPauseRenderingModule, StreamingPauseRendering);
 
+
 static TAutoConsoleVariable<int32> CVarRenderLastFrameInStreamingPause(
 	TEXT("r.RenderLastFrameInStreamingPause"),
 	1,
 	TEXT("If 1 the previous frame is displayed during streaming pause. If zero the screen is left black."),
 	ECVF_RenderThreadSafe);
 
+
 FStreamingPauseRenderingModule::FStreamingPauseRenderingModule()
-: Viewport(NULL)
-, ViewportWidget(NULL)
-, ViewportClient(NULL)
-{
-}
+	: Viewport(nullptr)
+	, ViewportWidget(nullptr)
+	, ViewportClient(nullptr)
+{ }
+
 
 void FStreamingPauseRenderingModule::StartupModule()
 {
@@ -33,18 +32,19 @@ void FStreamingPauseRenderingModule::StartupModule()
 	GEngine->RegisterEndStreamingPauseRenderingDelegate(&EndDelegate);
 }
 
+
 void FStreamingPauseRenderingModule::ShutdownModule()
 {
 	BeginDelegate.Unbind();
 	EndDelegate.Unbind();
 	if( GEngine )
 	{
-		GEngine->RegisterBeginStreamingPauseRenderingDelegate(NULL);
-		GEngine->RegisterEndStreamingPauseRenderingDelegate(NULL);
+		GEngine->RegisterBeginStreamingPauseRenderingDelegate(nullptr);
+		GEngine->RegisterEndStreamingPauseRenderingDelegate(nullptr);
 	}
 }
 
-/** Enqueue the streaming pause to suspend rendering during blocking load. */
+
 void FStreamingPauseRenderingModule::BeginStreamingPause( FViewport* GameViewport )
 {
 	//Create the viewport widget and add a throbber.
@@ -77,7 +77,7 @@ void FStreamingPauseRenderingModule::BeginStreamingPause( FViewport* GameViewpor
 
 		Viewport->EnqueueBeginRenderFrame();
 
-		FCanvas Canvas(Viewport.Get(), NULL, ViewportClient->GetWorld(), ViewportClient->GetWorld()->FeatureLevel);
+		FCanvas Canvas(Viewport.Get(), nullptr, ViewportClient->GetWorld(), ViewportClient->GetWorld()->FeatureLevel);
 		{
 			ViewportClient->Draw(Viewport.Get(), &Canvas);
 		}
@@ -103,15 +103,15 @@ void FStreamingPauseRenderingModule::BeginStreamingPause( FViewport* GameViewpor
 	GetMoviePlayer()->PlayMovie();
 }
 
-/** Enqueue the streaming pause to resume rendering after blocking load is completed. */
+
 void FStreamingPauseRenderingModule::EndStreamingPause()
 {
 	//Stop rendering the loading screen and resume
 	GetMoviePlayer()->WaitForMovieToFinish();
 
-	ViewportWidget = NULL;
-	Viewport = NULL;
-	ViewportClient = NULL;
+	ViewportWidget = nullptr;
+	Viewport = nullptr;
+	ViewportClient = nullptr;
 	
 	FlushRenderingCommands();
 }
\ No newline at end of file
diff --git a/Engine/Source/Runtime/StreamingPauseRendering/Public/StreamingPauseRendering.h b/Engine/Source/Runtime/StreamingPauseRendering/Public/StreamingPauseRendering.h
--- a/Engine/Source/Runtime/StreamingPauseRendering/Public/StreamingPauseRendering.h
+++ b/Engine/Source/Runtime/StreamingPauseRendering/Public/StreamingPauseRendering.h
@@ -2,30 +2,40 @@
 
 #pragma once
 
-/*=============================================================================
-	StreamingPauseRendering.h: Streaming pause definitions.
-=============================================================================*/
-
-#ifndef __STREAMINGPAUSERENDERING_H__
-#define __STREAMINGPAUSERENDERING_H__
-
 #include "Engine.h"
 #include "ModuleInterface.h"
 
+
 /** 
  * Module handling default behavior for streaming pause rendering. 
  * Games can override by calling RegisterBegin/EndStreamingPauseDelegate with their own delegates.
  */
-class FStreamingPauseRenderingModule : public IModuleInterface
+class FStreamingPauseRenderingModule
+	: public IModuleInterface
 {
 public:	
+
+	/** Default constructor. */
 	FStreamingPauseRenderingModule();
 
+	/**
+	 * Enqueue the streaming pause to suspend rendering during blocking load.
+	 */
+	virtual void BeginStreamingPause( class FViewport* Viewport );
+
+	/**
+	 * Enqueue the streaming pause to resume rendering after blocking load is completed.
+	 */
+	virtual void EndStreamingPause();
+
+public:
+
+	// IModuleInterface interface
+
 	virtual void StartupModule();	
 	virtual void ShutdownModule();
 
-	virtual void BeginStreamingPause( class FViewport* Viewport );
-	virtual void EndStreamingPause();
+public:
 
 	TSharedPtr<FSceneViewport>		Viewport;
 	TSharedPtr<SViewport>			ViewportWidget;
@@ -33,8 +43,7 @@ public:
 
 	/** Delegate providing default functionality for beginning streaming pause. */
 	FBeginStreamingPauseDelegate BeginDelegate;
+
 	/** Delegate providing default functionality for ending streaming pause. */
 	FEndStreamingPauseDelegate EndDelegate;
 };
-
-#endif // __STREAMINGPAUSERENDERING_H__
diff --git a/Engine/Source/Runtime/VectorVM/Private/VectorVM.cpp b/Engine/Source/Runtime/VectorVM/Private/VectorVM.cpp
--- a/Engine/Source/Runtime/VectorVM/Private/VectorVM.cpp
+++ b/Engine/Source/Runtime/VectorVM/Private/VectorVM.cpp
@@ -1,12 +1,9 @@
 // Copyright 1998-2014 Epic Games, Inc. All Rights Reserved.
 
-/*==============================================================================
-	VectorVM.cpp: Implementation of the vector virtual machine.
-==============================================================================*/
-
 #include "VectorVMPrivate.h"
 #include "ModuleManager.h"
 
+
 IMPLEMENT_MODULE(FDefaultModuleImpl, VectorVM);
 
 DEFINE_LOG_CATEGORY_STATIC(LogVectorVM, All, All);
diff --git a/Engine/Source/Runtime/VectorVM/Private/VectorVMPrivate.h b/Engine/Source/Runtime/VectorVM/Private/VectorVMPrivate.h
--- a/Engine/Source/Runtime/VectorVM/Private/VectorVMPrivate.h
+++ b/Engine/Source/Runtime/VectorVM/Private/VectorVMPrivate.h
@@ -1,12 +1,10 @@
 // Copyright 1998-2014 Epic Games, Inc. All Rights Reserved.
 
-/*==============================================================================
-	VectorVMPrivate.h: Private definitions for the vector virtual machine.
-==============================================================================*/
-
 #pragma once
+
 #include "VectorVM.h"
 
+
 namespace VectorVM
 {
 	/** Constants. */
diff --git a/Engine/Source/Runtime/VectorVM/Public/VectorVM.h b/Engine/Source/Runtime/VectorVM/Public/VectorVM.h
--- a/Engine/Source/Runtime/VectorVM/Public/VectorVM.h
+++ b/Engine/Source/Runtime/VectorVM/Public/VectorVM.h
@@ -1,12 +1,8 @@
 // Copyright 1998-2014 Epic Games, Inc. All Rights Reserved.
 
-/*==============================================================================
-	VectorVM.h: Public interface to the vector virtual machine.
-==============================================================================*/
-
 #pragma once
-#include "Core.h"
 
+#include "Core.h"
 
 
 namespace VectorVM
@@ -24,7 +20,6 @@ namespace VectorVM
 		MaxRegisters = NumTempRegisters + MaxInputRegisters + MaxOutputRegisters + MaxConstants
 	};
 
-
 	/** List of opcodes supported by the VM. */
 	namespace EOp
 	{
