UE-5596 Unable to override keybinding

diff --git a/Engine/Source/Editor/InputBindingEditor/Private/Widgets/SGestureEditBox.cpp b/Engine/Source/Editor/InputBindingEditor/Private/Widgets/SGestureEditBox.cpp
--- a/Engine/Source/Editor/InputBindingEditor/Private/Widgets/SGestureEditBox.cpp
+++ b/Engine/Source/Editor/InputBindingEditor/Private/Widgets/SGestureEditBox.cpp
@@ -23,6 +23,7 @@ void SGestureEditBox::Construct( const FArguments& InArgs, TSharedPtr<FGestureTr
 		SAssignNew( ConflictPopup, SMenuAnchor )
 		.Placement(MenuPlacement_ComboBox)
 		.OnGetMenuContent( this, &SGestureEditBox::OnGetContentForConflictPopup )
+		.OnMenuOpenChanged(this, &SGestureEditBox::OnConflictPopupOpenChanged)
 		[
 			SNew( SBox )
 			.WidthOverride( 200.0f )
@@ -127,8 +128,6 @@ void SGestureEditBox::OnGestureEditingStopped()
 	{
 		GestureEditor->CommitNewGesture();
 	}
-
-	ConflictPopup->SetIsOpen(false);
 }
 
 
@@ -136,7 +135,7 @@ void SGestureEditBox::OnGestureChanged()
 {
 	if( GestureEditor->HasConflict() )
 	{
-		ConflictPopup->SetIsOpen(true, false);
+		ConflictPopup->SetIsOpen(true, true);
 	}
 	else
 	{
@@ -151,7 +150,6 @@ void SGestureEditBox::OnGestureChanged()
 	}
 }
 
-
 EVisibility SGestureEditBox::GetGestureRemoveButtonVisibility() const
 {
 	if( !GestureEditor->IsEditing() && GestureEditor->IsActiveGestureValid() )
@@ -245,6 +243,13 @@ TSharedRef<SWidget> SGestureEditBox::OnGetContentForConflictPopup()
 		];
 }
 
+void SGestureEditBox::OnConflictPopupOpenChanged(bool bIsOpen)
+{
+	if(!bIsOpen)
+	{
+		GestureEditor->StopEditing();
+	}
+}
 
 EVisibility SGestureEditBox::GetNotificationVisibility() const
 {
@@ -268,5 +273,4 @@ FReply SGestureEditBox::OnMouseButtonDown( const FGeometry& MyGeometry, const FP
 	return FReply::Handled();
 }
 
-
 #undef LOCTEXT_NAMESPACE
diff --git a/Engine/Source/Editor/InputBindingEditor/Private/Widgets/SGestureEditBox.h b/Engine/Source/Editor/InputBindingEditor/Private/Widgets/SGestureEditBox.h
--- a/Engine/Source/Editor/InputBindingEditor/Private/Widgets/SGestureEditBox.h
+++ b/Engine/Source/Editor/InputBindingEditor/Private/Widgets/SGestureEditBox.h
@@ -54,6 +54,9 @@ private:
 	/** @return content to be shown in the key binding conflict pop-up */
 	TSharedRef<SWidget> OnGetContentForConflictPopup();
 
+	/** Called when conflict popup opens or closes */
+	void OnConflictPopupOpenChanged(bool bIsOpen);
+
 	/** @return The visibility of the duplicate binding notification area */
 	EVisibility GetNotificationVisibility() const; 
 	virtual FReply OnMouseButtonDown( const FGeometry& MyGeometry, const FPointerEvent& InMouseEvent );
