Improve alignment in the editor settings keyboard shortcuts editor

diff --git a/Engine/Source/Editor/InputBindingEditor/Private/Widgets/SChordTreeItem.cpp b/Engine/Source/Editor/InputBindingEditor/Private/Widgets/SChordTreeItem.cpp
--- a/Engine/Source/Editor/InputBindingEditor/Private/Widgets/SChordTreeItem.cpp
+++ b/Engine/Source/Editor/InputBindingEditor/Private/Widgets/SChordTreeItem.cpp
@@ -28,8 +28,9 @@ TSharedRef<SWidget> SChordTreeItem::GenerateWidgetForColumn(const FName& ColumnN
 				SNew(SExpanderArrow, SharedThis(this))
 			]
 
-		+ SHorizontalBox::Slot()
+			+ SHorizontalBox::Slot()
 			.AutoWidth()
+			.Padding(0.0f, 2.0f, 0.0f, 2.0f)
 			.VAlign(VAlign_Center)
 			[
 				SNew(SVerticalBox)
@@ -43,14 +44,15 @@ TSharedRef<SWidget> SChordTreeItem::GenerateWidgetForColumn(const FName& ColumnN
 				]
 
 				+ SVerticalBox::Slot()
-					.Padding(0.0f, 0.0f, 0.0f, 5.0f)
-					.AutoHeight()
-					[
-						SNew(STextBlock)
-						.Font(FEditorStyle::GetFontStyle("InputBindingEditor.SmallFont"))
-						.ColorAndOpacity(FLinearColor::Gray)
-						.Text(TreeItem->CommandInfo->GetDescription())
-					]
+				.Padding(0.0f, 2.0f, 0.0f, 0.0f)
+				.AutoHeight()
+				[
+					SNew(STextBlock)
+					.Font(FEditorStyle::GetFontStyle("InputBindingEditor.SmallFont"))
+					.ColorAndOpacity(FLinearColor::Gray)
+					.Text(TreeItem->CommandInfo->GetDescription())
+					//.Visibility(TreeItem->CommandInfo->GetDescription().IsEmptyOrWhitespace() ? EVisibility::Collapsed : EVisibility::Visible)
+				]
 			];
 	}
 	else if (ColumnName == "Binding")
