Workaround for memory stomp

diff --git a/Engine/Source/Editor/InputBindingEditor/Private/Widgets/SInputBindingEditorPanel.cpp b/Engine/Source/Editor/InputBindingEditor/Private/Widgets/SInputBindingEditorPanel.cpp
--- a/Engine/Source/Editor/InputBindingEditor/Private/Widgets/SInputBindingEditorPanel.cpp
+++ b/Engine/Source/Editor/InputBindingEditor/Private/Widgets/SInputBindingEditorPanel.cpp
@@ -30,20 +30,21 @@ void FInputBindingEditorPanel::Tick(float DeltaSeconds)
 {
 	if(bUpdateRequested)
 	{
+		bUpdateRequested = false;
+
 		UpdateContextMasterList();
 
 		if (DetailBuilder)
 		{
-			DetailBuilder->ForceRefreshDetails();
-			// Force refreshing is going to invalidate the detail builder anyway
-			DetailBuilder = nullptr;
 			FBindingContext::CommandsChanged.RemoveAll(this);
-		}
 
-		bUpdateRequested = false;
-	}
+			DetailBuilder->ForceRefreshDetails();
 
-	
+			// DetailBuilder->ForceRefreshDetails() can cause the this pointer to become invalidated so return
+			// right after this call to avoid any accidental usage of member variables after this
+			return;
+		}
+	}
 }
 
 void FInputBindingEditorPanel::UpdateContextMasterList()
