Fix ensure with input binding editor panel

diff --git a/Engine/Source/Editor/InputBindingEditor/Private/InputBindingEditorModule.cpp b/Engine/Source/Editor/InputBindingEditor/Private/InputBindingEditorModule.cpp
--- a/Engine/Source/Editor/InputBindingEditor/Private/InputBindingEditorModule.cpp
+++ b/Engine/Source/Editor/InputBindingEditor/Private/InputBindingEditorModule.cpp
@@ -5,7 +5,6 @@
 #include "Interfaces/IInputBindingEditorModule.h"
 #include "Framework/Commands/InputBindingManager.h"
 #include "Misc/Paths.h"
-#include "Widgets/SInputBindingEditorPanel.h"
 #include "UnrealEdMisc.h"
 #include "Logging/MessageLog.h"
 #include "Dialogs/Dialogs.h"
@@ -17,36 +16,224 @@
 #include "HAL/FileManager.h"
 #include "Misc/ConfigCacheIni.h"
 
+#include "Widgets/Images/SImage.h"
+#include "Widgets/Input/SButton.h"
+#include "Widgets/SChordEditBox.h"
+#include "Widgets/Input/SSearchBox.h"
+#include "Widgets/Text/STextBlock.h"
+#include "DetailLayoutBuilder.h"
+#include "DetailCategoryBuilder.h"
+#include "DetailWidgetRow.h"
+
 #define LOCTEXT_NAMESPACE "InputBindingEditor"
 
 static FName SettingsModuleName("Settings");
 static FName PropertyEditorModuleName("PropertyEditor");
 
+/**
+ * A gesture sort functor.  Sorts by name or gesture and ascending or descending
+ */
+struct FChordSort
+{
+	FChordSort( bool bInSortName, bool bInSortUp )
+		: bSortName( bInSortName )
+		, bSortUp( bInSortUp )
+	{ }
+
+	bool operator()( const TSharedPtr<FUICommandInfo>& A, const TSharedPtr<FUICommandInfo>& B ) const
+	{
+		if( bSortName )
+		{
+			bool bResult = A->GetLabel().CompareTo( B->GetLabel() ) == -1;
+			return bSortUp ? !bResult : bResult;
+		}
+		else
+		{
+			// Sort by binding
+			bool bResult = A->GetInputText().CompareTo( B->GetInputText() ) == -1;
+			return bSortUp ? !bResult : bResult;
+		}
+	}
+
+private:
+
+	/** Whether or not to sort by name.  If false we sort by binding. */
+	bool bSortName;
+
+	/** Whether or not to sort up.  If false we sort down. */
+	bool bSortUp;
+};
+
+/**
+* An item for the chord tree view
+*/
+struct FChordTreeItem
+{
+	// Note these are mutually exclusive
+	TWeakPtr<FBindingContext> BindingContext;
+	TSharedPtr<FUICommandInfo> CommandInfo;
+
+	TSharedPtr<FBindingContext> GetBindingContext() { return BindingContext.Pin(); }
+
+	bool IsContext() const { return BindingContext.IsValid(); }
+	bool IsCommand() const { return CommandInfo.IsValid(); }
+};
+
 class FEditorKeyboardShortcutSettings : public IDetailCustomization
 {
 public:
+
+	FEditorKeyboardShortcutSettings() :
+		bUpdateRequested(false),
+		DetailBuilder(nullptr)
+	{
+	}
+
+	virtual ~FEditorKeyboardShortcutSettings()
+	{
+		FBindingContext::CommandsChanged.RemoveAll(this);
+		FInputBindingManager::Get().SaveInputBindings();
+	}
+
 	static TSharedRef<IDetailCustomization> MakeInstance()
 	{
 		return MakeShareable( new FEditorKeyboardShortcutSettings );
 	}
 
-	virtual void PendingDelete()
+	virtual void CustomizeDetails(IDetailLayoutBuilder& InDetailBuilder) override
 	{
-		if (EditorPanel.IsValid())
+		DetailBuilder = &InDetailBuilder;
+
+		UpdateContextMasterList();
+		UpdateUI();
+
+		FBindingContext::CommandsChanged.AddSP( SharedThis( this ), &FEditorKeyboardShortcutSettings::OnCommandsChanged );
+	}
+
+	/** Updates the master context list with new commands. */
+	void UpdateContextMasterList()
+	{
+		TArray< TSharedPtr<FBindingContext> > Contexts;
+		FInputBindingManager::Get().GetKnownInputContexts( Contexts );
+
+		struct FContextNameSort
+		{
+			bool operator()( const TSharedPtr<FBindingContext>& A, const TSharedPtr<FBindingContext>& B ) const
+			{
+				return A->GetContextDesc().CompareTo( B->GetContextDesc() ) == -1;
+			}
+		};
+		Contexts.Sort( FContextNameSort() );
+
+		/** List of all known contexts. */
+		ContextMasterList.Reset(Contexts.Num());
+
+		for (const TSharedPtr<FBindingContext>& Context : Contexts)
+		{
+			TSharedRef<FChordTreeItem> TreeItem( new FChordTreeItem );
+			TreeItem->BindingContext = Context;
+			ContextMasterList.Add( TreeItem );
+		}
+	}
+
+	void ForceRefreshDetails()
+	{
+		bUpdateRequested = false;
+		UpdateContextMasterList();
+
+		if (DetailBuilder)
+		{
+			FBindingContext::CommandsChanged.RemoveAll(this);
+
+			IDetailLayoutBuilder* DetailBuilderPtr = DetailBuilder;
+			DetailBuilder = nullptr;
+			DetailBuilderPtr->ForceRefreshDetails();
+		}
+	}
+
+	void OnCommandsChanged(const FBindingContext& ContextThatChanged)
+	{
+		if (!bUpdateRequested)
 		{
-			ensure(EditorPanel.IsUnique());
-			EditorPanel.Reset();
+			bUpdateRequested = true;
+			GEditor->GetTimerManager()->SetTimerForNextTick(FTimerDelegate::CreateSP(this, &FEditorKeyboardShortcutSettings::ForceRefreshDetails));
 		}
 	}
 
-	virtual void CustomizeDetails( IDetailLayoutBuilder& DetailBuilder )
+	void UpdateUI()
 	{
-		EditorPanel = MakeShareable( new FInputBindingEditorPanel );
+		for (TSharedPtr<FChordTreeItem>& TreeItem : ContextMasterList)
+		{
+			check(TreeItem->IsContext());
+
+			IDetailCategoryBuilder& CategoryBuilder = DetailBuilder->EditCategory(TreeItem->GetBindingContext()->GetContextName(), TreeItem->GetBindingContext()->GetContextDesc());
+
+			TArray<TSharedPtr<FUICommandInfo>> Commands;
+			GetCommandsForContext(TreeItem, Commands);
+
+			for(TSharedPtr<FUICommandInfo>& CommandInfo : Commands)
+			{
+				FDetailWidgetRow& Row = CategoryBuilder.AddCustomRow(CommandInfo->GetLabel());
+				Row.NameContent()
+				.MaxDesiredWidth(0)
+				.MinDesiredWidth(500)
+				[
+					SNew(SVerticalBox)
+					+ SVerticalBox::Slot()
+					.AutoHeight()
+					[
+						SNew(STextBlock)
+						.Text(CommandInfo->GetLabel())
+						.ToolTipText(CommandInfo->GetDescription())
+					]
+					+ SVerticalBox::Slot()
+					.Padding(0.0f, 3.0f, 0.0f, 3.0f)
+					.AutoHeight()
+					[
+						SNew(STextBlock)
+						.Font(IDetailLayoutBuilder::GetDetailFont())
+						.ColorAndOpacity(FLinearColor::Gray)
+						.Text(CommandInfo->GetDescription())
+					]
+				];
+
+				Row.ValueContent()
+				.MaxDesiredWidth(200)
+				.MinDesiredWidth(200)
+				.VAlign(VAlign_Center)
+				[
+					SNew(SHorizontalBox)
+
+					+ SHorizontalBox::Slot()
+					.AutoWidth()
+					.Padding(1.0f, 0.0f, 9.0f, 0.0f)
+					[
+						SNew(SChordEditBox, CommandInfo, EMultipleKeyBindingIndex::Primary)
+					]
+					+ SHorizontalBox::Slot()
+					.AutoWidth()
+					[
+						SNew(SChordEditBox, CommandInfo, EMultipleKeyBindingIndex::Secondary)
+					]
+				];
+			}
+		}
+	}
 
-		EditorPanel->Initialize(DetailBuilder);
+	void GetCommandsForContext(TSharedPtr<FChordTreeItem> InTreeItem, TArray< TSharedPtr< FUICommandInfo > >& OutChildren)
+	{
+		if (InTreeItem->IsContext())
+		{
+			FInputBindingManager::Get().GetCommandInfosFromContext(InTreeItem->GetBindingContext()->GetContextName(), OutChildren);
+			OutChildren.Sort(FChordSort(true, false));
+		}
 	}
+
 private:
-	TSharedPtr<FInputBindingEditorPanel> EditorPanel;
+	bool bUpdateRequested;
+	IDetailLayoutBuilder* DetailBuilder;
+	/** List of all known contexts. */
+	TArray< TSharedPtr<FChordTreeItem> > ContextMasterList;
 };
 
 class FInputBindingEditorModule
diff --git a/Engine/Source/Editor/InputBindingEditor/Private/Widgets/SInputBindingEditorPanel.cpp b/Engine/Source/Editor/InputBindingEditor/Private/Widgets/SInputBindingEditorPanel.cpp
deleted file mode 100644
diff --git a/Engine/Source/Editor/InputBindingEditor/Private/Widgets/SInputBindingEditorPanel.h b/Engine/Source/Editor/InputBindingEditor/Private/Widgets/SInputBindingEditorPanel.h
deleted file mode 100644
