Fixed issue with SChordEditBox preventing keybinding chords from being overwritten on Linux due to prematurely-cleared data

Linux uses X11, which is ASync, to handle focus whereas Windows handles things synchronously. This causes prematurely cleared data in KeyDown with IsTyping.

diff --git a/Engine/Source/Editor/InputBindingEditor/Private/Widgets/SChordEditBox.cpp b/Engine/Source/Editor/InputBindingEditor/Private/Widgets/SChordEditBox.cpp
--- a/Engine/Source/Editor/InputBindingEditor/Private/Widgets/SChordEditBox.cpp
+++ b/Engine/Source/Editor/InputBindingEditor/Private/Widgets/SChordEditBox.cpp
@@ -110,9 +110,9 @@ FText SChordEditBox::GetNotificationMessage() const
 
 void SChordEditBox::OnChordEditorLostFocus()
 {
-	if( (!ChordAcceptButton.IsValid() || ChordAcceptButton->HasMouseCapture() == false) && !ChordEditor->IsTyping() )
+	if( (!ChordAcceptButton.IsValid() || ChordAcceptButton->HasMouseCapture() == false) && !ChordEditor->IsTyping() && !ChordEditor->HasConflict() )
 	{
-		if( ChordEditor->IsEditing() && ChordEditor->IsEditedChordValid() && !ChordEditor->HasConflict() )
+		if( ChordEditor->IsEditing() && ChordEditor->IsEditedChordValid() )
 		{
 			ChordEditor->CommitNewChord();
 		}
