Fixed some issues found when running includetool.

diff --git a/Engine/Source/Editor/DeviceProfileEditor/Private/DetailsPanel/SDeviceProfileDetailsPanel.h b/Engine/Source/Editor/DeviceProfileEditor/Private/DetailsPanel/SDeviceProfileDetailsPanel.h
--- a/Engine/Source/Editor/DeviceProfileEditor/Private/DetailsPanel/SDeviceProfileDetailsPanel.h
+++ b/Engine/Source/Editor/DeviceProfileEditor/Private/DetailsPanel/SDeviceProfileDetailsPanel.h
@@ -2,7 +2,6 @@
 
 #pragma once
 
-#include "SDeviceProfileDetailsPanel.h"
 #include "Templates/SharedPointer.h"
 #include "UObject/WeakObjectPtr.h"
 #include "UObject/WeakObjectPtrTemplates.h"
diff --git a/Engine/Source/Editor/InputBindingEditor/Private/Widgets/SChordEditor.h b/Engine/Source/Editor/InputBindingEditor/Private/Widgets/SChordEditor.h
--- a/Engine/Source/Editor/InputBindingEditor/Private/Widgets/SChordEditor.h
+++ b/Engine/Source/Editor/InputBindingEditor/Private/Widgets/SChordEditor.h
@@ -7,7 +7,6 @@
 #include "Framework/Commands/UICommandInfo.h"
 #include "Input/Reply.h"
 #include "Internationalization/Text.h"
-#include "SChordEditor.h"
 #include "Templates/SharedPointer.h"
 #include "Widgets/DeclarativeSyntaxSupport.h"
 #include "Widgets/Input/SEditableText.h"
