Replaced usages of "master" in InputBindingEditorModule

diff --git a/Engine/Source/Editor/InputBindingEditor/Private/InputBindingEditorModule.cpp b/Engine/Source/Editor/InputBindingEditor/Private/InputBindingEditorModule.cpp
--- a/Engine/Source/Editor/InputBindingEditor/Private/InputBindingEditorModule.cpp
+++ b/Engine/Source/Editor/InputBindingEditor/Private/InputBindingEditorModule.cpp
@@ -154,14 +154,14 @@ public:
 	{
 		DetailBuilder = &InDetailBuilder;
 
-		UpdateContextMasterList();
+		UpdateContextList();
 		UpdateUI();
 
 		FBindingContext::CommandsChanged.AddSP( SharedThis( this ), &FEditorKeyboardShortcutSettings::OnCommandsChanged );
 	}
 
-	/** Updates the master context list with new commands. */
-	void UpdateContextMasterList()
+	/** Updates the context list with new commands. */
+	void UpdateContextList()
 	{
 		TArray< TSharedPtr<FBindingContext> > Contexts;
 		FInputBindingManager::Get().GetKnownInputContexts( Contexts );
@@ -176,20 +176,20 @@ public:
 		Contexts.Sort( FContextNameSort() );
 
 		/** List of all known contexts. */
-		ContextMasterList.Reset(Contexts.Num());
+		ContextList.Reset(Contexts.Num());
 
 		for (const TSharedPtr<FBindingContext>& Context : Contexts)
 		{
 			TSharedRef<FChordTreeItem> TreeItem( new FChordTreeItem );
 			TreeItem->BindingContext = Context;
-			ContextMasterList.Add( TreeItem );
+			ContextList.Add( TreeItem );
 		}
 	}
 
 	void ForceRefreshDetails()
 	{
 		bUpdateRequested = false;
-		UpdateContextMasterList();
+		UpdateContextList();
 
 		if (DetailBuilder)
 		{
@@ -212,7 +212,7 @@ public:
 
 	void UpdateUI()
 	{
-		for (TSharedPtr<FChordTreeItem>& TreeItem : ContextMasterList)
+		for (TSharedPtr<FChordTreeItem>& TreeItem : ContextList)
 		{
 			check(TreeItem->IsContext());
 
@@ -314,7 +314,7 @@ private:
 	bool bUpdateRequested;
 	IDetailLayoutBuilder* DetailBuilder;
 	/** List of all known contexts. */
-	TArray< TSharedPtr<FChordTreeItem> > ContextMasterList;
+	TArray< TSharedPtr<FChordTreeItem> > ContextList;
 };
 
 class FInputBindingEditorModule
