Removed FName look ups that were showing in profiles by making FName static variables.

diff --git a/Engine/Source/Editor/InputBindingEditor/Private/Widgets/SChordEditBox.cpp b/Engine/Source/Editor/InputBindingEditor/Private/Widgets/SChordEditBox.cpp
--- a/Engine/Source/Editor/InputBindingEditor/Private/Widgets/SChordEditBox.cpp
+++ b/Engine/Source/Editor/InputBindingEditor/Private/Widgets/SChordEditBox.cpp
@@ -37,11 +37,14 @@ struct FGeometry;
 
 void SChordEditBox::Construct( const FArguments& InArgs, TSharedPtr<FUICommandInfo> InputCommand, EMultipleKeyBindingIndex InChordIndex)
 {
-	BorderImageNormal = FAppStyle::GetBrush( "EditableTextBox.Background.Normal" );
-	BorderImageHovered = FAppStyle::GetBrush( "EditableTextBox.Background.Hovered" );
-	BorderImageFocused = FAppStyle::GetBrush( "EditableTextBox.Background.Focused" );
-
 	static const FName InvertedForegroundName("InvertedForeground");
+	static const FName EditableTextBoBackgroundNormal("EditableTextBox.Background.Normal");
+	static const FName EditableTextBoBackgroundHovered("EditableTextBox.Background.Hovered");
+	static const FName EditableTextBoBackgroundFocused("EditableTextBox.Background.Focused");
+
+	BorderImageNormal = FAppStyle::GetBrush(EditableTextBoBackgroundNormal);
+	BorderImageHovered = FAppStyle::GetBrush(EditableTextBoBackgroundHovered);
+	BorderImageFocused = FAppStyle::GetBrush(EditableTextBoBackgroundFocused);
 
 	ChildSlot
 	[
