Respect command filtering in the input binding manager

#rnx

diff --git a/Engine/Source/Editor/InputBindingEditor/Private/InputBindingEditorModule.cpp b/Engine/Source/Editor/InputBindingEditor/Private/InputBindingEditorModule.cpp
--- a/Engine/Source/Editor/InputBindingEditor/Private/InputBindingEditorModule.cpp
+++ b/Engine/Source/Editor/InputBindingEditor/Private/InputBindingEditorModule.cpp
@@ -77,17 +77,17 @@ struct FChordSort
 			// Sort by command bundle, and then by command label. If a command has no bundle,
 			// it will compare its label to the other command's bundle.
 			const int32 CompareResult = GetPrimaryTextForCommand(A).CompareTo(GetPrimaryTextForCommand(B));
-			bool bFinalResult = CompareResult == -1;
+			bool bFinalResult = CompareResult < 0;
 			if (CompareResult == 0)
 			{
-				bFinalResult = A->GetLabel().CompareTo(B->GetLabel()) == -1;
+				bFinalResult = A->GetLabel().CompareTo(B->GetLabel()) < 0;
 			}
 			return bSortUp ? !bFinalResult : bFinalResult;
 		}
 		else
 		{
 			// Sort by binding
-			bool bResult = A->GetInputText().CompareTo( B->GetInputText() ) == -1;
+			bool bResult = A->GetInputText().CompareTo( B->GetInputText() ) < 0;
 			return bSortUp ? !bResult : bResult;
 		}
 	}
@@ -163,17 +163,21 @@ public:
 	/** Updates the context list with new commands. */
 	void UpdateContextList()
 	{
+		FInputBindingManager& InputBindingManager = FInputBindingManager::Get();
+
 		TArray< TSharedPtr<FBindingContext> > Contexts;
-		FInputBindingManager::Get().GetKnownInputContexts( Contexts );
+		InputBindingManager.GetKnownInputContexts( Contexts );
 
-		struct FContextNameSort
+		// Filter to allowed bindings
+		Contexts.RemoveAll([&InputBindingManager](const TSharedPtr<FBindingContext>& Context)
 		{
-			bool operator()( const TSharedPtr<FBindingContext>& A, const TSharedPtr<FBindingContext>& B ) const
-			{
-				return A->GetContextDesc().CompareTo( B->GetContextDesc() ) == -1;
-			}
-		};
-		Contexts.Sort( FContextNameSort() );
+			return !InputBindingManager.CommandPassesFilter(FName(), Context->GetContextName());
+		});
+
+		Contexts.Sort([](const TSharedPtr<FBindingContext>& A, const TSharedPtr<FBindingContext>& B)
+		{
+			return A->GetContextDesc().CompareTo(B->GetContextDesc()) < 0;
+		});
 
 		/** List of all known contexts. */
 		ContextList.Reset(Contexts.Num());
@@ -212,6 +216,8 @@ public:
 
 	void UpdateUI()
 	{
+		FInputBindingManager& InputBindingManager = FInputBindingManager::Get();
+
 		for (TSharedPtr<FChordTreeItem>& TreeItem : ContextList)
 		{
 			check(TreeItem->IsContext());
@@ -219,7 +225,15 @@ public:
 			IDetailCategoryBuilder& CategoryBuilder = DetailBuilder->EditCategory(TreeItem->GetBindingContext()->GetContextName(), TreeItem->GetBindingContext()->GetContextDesc());
 
 			TArray<TSharedPtr<FUICommandInfo>> Commands;
-			GetCommandsForContext(TreeItem, Commands);
+			InputBindingManager.GetCommandInfosFromContext(TreeItem->GetBindingContext()->GetContextName(), Commands);
+
+			// Filter to allowed bindings
+			Commands.RemoveAll([&InputBindingManager](const TSharedPtr<FUICommandInfo>& CommandInfo)
+			{
+				return !InputBindingManager.CommandPassesFilter(CommandInfo->GetBindingContext(), CommandInfo->GetCommandName());
+			});
+			
+			Commands.Sort(FChordSort(true, false));
 
 			TMap<FName, IDetailGroup*> BundleMap;
 
@@ -301,15 +315,6 @@ public:
 		}
 	}
 
-	void GetCommandsForContext(TSharedPtr<FChordTreeItem> InTreeItem, TArray< TSharedPtr< FUICommandInfo > >& OutChildren)
-	{
-		if (InTreeItem->IsContext())
-		{
-			FInputBindingManager::Get().GetCommandInfosFromContext(InTreeItem->GetBindingContext()->GetContextName(), OutChildren);
-			OutChildren.Sort(FChordSort(true, false));
-		}
-	}
-
 private:
 	bool bUpdateRequested;
 	IDetailLayoutBuilder* DetailBuilder;
diff --git a/Engine/Source/Runtime/Slate/Private/Framework/Commands/InputBindingManager.cpp b/Engine/Source/Runtime/Slate/Private/Framework/Commands/InputBindingManager.cpp
--- a/Engine/Source/Runtime/Slate/Private/Framework/Commands/InputBindingManager.cpp
+++ b/Engine/Source/Runtime/Slate/Private/Framework/Commands/InputBindingManager.cpp
@@ -279,13 +279,8 @@ void FUserDefinedChords::RemoveAll()
 
 FInputBindingManager& FInputBindingManager::Get()
 {
-	static FInputBindingManager* Instance= NULL;
-	if( Instance == NULL )
-	{
-		Instance = new FInputBindingManager();
-	}
-
-	return *Instance;
+	static FInputBindingManager Instance;
+	return Instance;
 }
 
 
