Added debug info to the input binding UI

This is controlled by the "Input.Debug.ShowBindingNames" CVar, and defaults to false

diff --git a/Engine/Source/Editor/InputBindingEditor/Private/InputBindingEditorModule.cpp b/Engine/Source/Editor/InputBindingEditor/Private/InputBindingEditorModule.cpp
--- a/Engine/Source/Editor/InputBindingEditor/Private/InputBindingEditorModule.cpp
+++ b/Engine/Source/Editor/InputBindingEditor/Private/InputBindingEditorModule.cpp
@@ -15,6 +15,7 @@
 #include "Framework/Commands/InputBindingManager.h"
 #include "Framework/Commands/UICommandInfo.h"
 #include "HAL/FileManager.h"
+#include "HAL/IConsoleManager.h"
 #include "HAL/Platform.h"
 #include "HAL/PlatformCrt.h"
 #include "HAL/PlatformMisc.h"
@@ -57,8 +58,14 @@ class SWidget;
 
 #define LOCTEXT_NAMESPACE "InputBindingEditor"
 
-static FName SettingsModuleName("Settings");
-static FName PropertyEditorModuleName("PropertyEditor");
+namespace InputBindingEditorModule
+{
+static const FName SettingsModuleName("Settings");
+static const FName PropertyEditorModuleName("PropertyEditor");
+
+bool bShowBindingNames = false;
+static FAutoConsoleVariableRef CVarDebugBindingNames(TEXT("Input.Debug.ShowBindingNames"), bShowBindingNames, TEXT("True to show binding names in the input binding editor."));
+}
 
 /**
  * A gesture sort functor.  Sorts by name or gesture and ascending or descending
@@ -279,7 +286,17 @@ public:
 					[
 						SNew(STextBlock)
 						.Text(CommandInfo->GetLabel())
-						.ToolTipText(CommandInfo->GetDescription())
+						.ToolTipText_Lambda([CommandInfo]() -> FText
+						{
+							FText CommandInfoTooltip = CommandInfo->GetDescription();
+
+							if (InputBindingEditorModule::bShowBindingNames)
+							{
+								CommandInfoTooltip = FText::Format(LOCTEXT("CommandInfoDebugToolTip", "{0}\n\nBinding Context: {1}\nCommand Name: {2}"), CommandInfoTooltip, FText::FromName(CommandInfo->GetBindingContext()), FText::FromName(CommandInfo->GetCommandName()));
+							}
+
+							return CommandInfoTooltip;
+						})
 					]
 					+ SVerticalBox::Slot()
 					.Padding(0.0f, 3.0f, 0.0f, 3.0f)
@@ -330,9 +347,9 @@ public:
 	// IInputBindingEditorModule interface
 	virtual void StartupModule() override
 	{
-		ISettingsModule& SettingsModule = FModuleManager::LoadModuleChecked<ISettingsModule>(SettingsModuleName);
+		ISettingsModule& SettingsModule = FModuleManager::LoadModuleChecked<ISettingsModule>(InputBindingEditorModule::SettingsModuleName);
 
-		FPropertyEditorModule& PropertyEditor = FModuleManager::LoadModuleChecked<FPropertyEditorModule>(PropertyEditorModuleName);
+		FPropertyEditorModule& PropertyEditor = FModuleManager::LoadModuleChecked<FPropertyEditorModule>(InputBindingEditorModule::PropertyEditorModuleName);
 
 		EditorKeyboardShortcutSettingsName = UEditorKeyboardShortcutSettings::StaticClass()->GetFName();
 		PropertyEditor.RegisterCustomClassLayout(EditorKeyboardShortcutSettingsName, FOnGetDetailCustomizationInstance::CreateStatic(&FEditorKeyboardShortcutSettings::MakeInstance));
@@ -355,9 +372,9 @@ public:
 
 	virtual void ShutdownModule() override
 	{
-		if(FModuleManager::Get().IsModuleLoaded(PropertyEditorModuleName))
+		if(FModuleManager::Get().IsModuleLoaded(InputBindingEditorModule::PropertyEditorModuleName))
 		{
-			FPropertyEditorModule& PropertyEditor = FModuleManager::GetModuleChecked<FPropertyEditorModule>(PropertyEditorModuleName);
+			FPropertyEditorModule& PropertyEditor = FModuleManager::GetModuleChecked<FPropertyEditorModule>(InputBindingEditorModule::PropertyEditorModuleName);
 
 			PropertyEditor.UnregisterCustomClassLayout(EditorKeyboardShortcutSettingsName);
 		}
