Make keyboard shortcuts searchable in Editor Settings with the format: KeyBinding="F"

diff --git a/Engine/Source/Editor/InputBindingEditor/Private/InputBindingEditorModule.cpp b/Engine/Source/Editor/InputBindingEditor/Private/InputBindingEditorModule.cpp
--- a/Engine/Source/Editor/InputBindingEditor/Private/InputBindingEditorModule.cpp
+++ b/Engine/Source/Editor/InputBindingEditor/Private/InputBindingEditorModule.cpp
@@ -276,6 +276,42 @@ public:
 					Row = &CategoryBuilder.AddCustomRow(CommandInfo->GetLabel());
 				}
 
+				// Set up search filter (for i.e. KeyBinding="F")
+				const TSharedRef<const FInputChord> FirstInputChord = CommandInfo->GetActiveChord(EMultipleKeyBindingIndex::Primary);
+				const TSharedRef<const FInputChord> SecondInputChord = CommandInfo->GetActiveChord(EMultipleKeyBindingIndex::Secondary);
+
+				const bool bFirstChordValid = FirstInputChord->IsValidChord();
+				const bool bSecondChordValid = SecondInputChord->IsValidChord();
+
+				if (bFirstChordValid || bSecondChordValid)
+				{
+					static const FTextFormat KeyFormat = LOCTEXT("SearchKeyFilter", "KeyBinding=\"{0}\"");
+					static const FTextFormat TokenCombineFormat = INVTEXT("{0} {1}");
+
+					FText FirstKeyFilter, SecondKeyFilter;
+					if (bFirstChordValid)
+					{
+						FirstKeyFilter = FText::FormatOrdered(KeyFormat, FirstInputChord->GetInputText());
+					}
+
+					if (bSecondChordValid)
+					{
+						SecondKeyFilter = FText::FormatOrdered(KeyFormat, FirstInputChord->GetInputText());
+					}
+
+					FText FilterText = bFirstChordValid && bSecondChordValid
+						? FText::FormatOrdered(TokenCombineFormat, FirstKeyFilter, SecondKeyFilter)
+						: (bFirstChordValid ? FirstKeyFilter : SecondKeyFilter);
+
+					// Command label and similar string might already be stored
+					if (!Row->FilterTextString.IsEmpty())
+					{
+						FilterText = FText::FormatOrdered(TokenCombineFormat, Row->FilterTextString, FilterText);
+					}
+
+					Row->FilterString(FilterText);
+				}
+
 				Row->NameContent()
 				.MaxDesiredWidth(0)
 				.MinDesiredWidth(500)
