InputBindingEditor: Change the key binding conflict notification error text color to FStyleColors::Warning

- The previous red color had a contrast of 1.76:1 with the brackground, the new yellow color has a contrast of 4.99:1

#rnx

diff --git a/Engine/Source/Editor/InputBindingEditor/Private/Widgets/SChordEditBox.cpp b/Engine/Source/Editor/InputBindingEditor/Private/Widgets/SChordEditBox.cpp
--- a/Engine/Source/Editor/InputBindingEditor/Private/Widgets/SChordEditBox.cpp
+++ b/Engine/Source/Editor/InputBindingEditor/Private/Widgets/SChordEditBox.cpp
@@ -13,6 +13,7 @@
 #include "SlotBase.h"
 #include "Styling/AppStyle.h"
 #include "Styling/SlateColor.h"
+#include "Styling/StyleColors.h"
 #include "Types/SlateEnums.h"
 #include "Types/SlateStructs.h"
 #include "UObject/NameTypes.h"
@@ -229,7 +230,7 @@ TSharedRef<SWidget> SChordEditBox::OnGetContentForConflictPopup()
 				[
 					SNew( STextBlock )
 						.WrapTextAt(200.0f)
-						.ColorAndOpacity( FLinearColor( .75f, 0.0f, 0.0f, 1.0f ) )
+						.ColorAndOpacity( FStyleColors::Warning )
 						.Text( this, &SChordEditBox::GetNotificationMessage )
 						.Visibility( this, &SChordEditBox::GetNotificationVisibility )
 				]
