EditorUsability : InputBindingEditor
- Improved "Keyboard Shortcuts" editable input UI to follow unreal theming

#Synced-CL 48229788

diff --git a/Engine/Source/Editor/InputBindingEditor/Private/Widgets/SChordEditBox.cpp b/Engine/Source/Editor/InputBindingEditor/Private/Widgets/SChordEditBox.cpp
--- a/Engine/Source/Editor/InputBindingEditor/Private/Widgets/SChordEditBox.cpp
+++ b/Engine/Source/Editor/InputBindingEditor/Private/Widgets/SChordEditBox.cpp
@@ -36,63 +36,64 @@ struct FGeometry;
 /* SChordEditBox interface
  *****************************************************************************/
 
-void SChordEditBox::Construct( const FArguments& InArgs, TSharedPtr<FUICommandInfo> InputCommand, EMultipleKeyBindingIndex InChordIndex)
+void SChordEditBox::Construct(const FArguments& InArgs, TSharedPtr<FUICommandInfo> InputCommand, EMultipleKeyBindingIndex InChordIndex)
 {
-	static const FName InvertedForegroundName("InvertedForeground");
-	static const FName EditableTextBoBackgroundNormal("EditableTextBox.Background.Normal");
-	static const FName EditableTextBoBackgroundHovered("EditableTextBox.Background.Hovered");
-	static const FName EditableTextBoBackgroundFocused("EditableTextBox.Background.Focused");
+	// Match EditableTextBox style
+	const FEditableTextBoxStyle& EditableTextBoxStyle = FCoreStyle::Get().GetWidgetStyle<FEditableTextBoxStyle>("NormalEditableTextBox");
 
-	BorderImageNormal = FAppStyle::GetBrush(EditableTextBoBackgroundNormal);
-	BorderImageHovered = FAppStyle::GetBrush(EditableTextBoBackgroundHovered);
-	BorderImageFocused = FAppStyle::GetBrush(EditableTextBoBackgroundFocused);
+	BorderImageNormal = &EditableTextBoxStyle.BackgroundImageNormal;
+	BorderImageHovered = &EditableTextBoxStyle.BackgroundImageHovered;
+	BorderImageFocused = &EditableTextBoxStyle.BackgroundImageFocused;
 
 	ChildSlot
 	[
-		SAssignNew( ConflictPopup, SMenuAnchor )
+		SAssignNew(ConflictPopup, SMenuAnchor)
 		.Placement(MenuPlacement_ComboBox)
-		.OnGetMenuContent( this, &SChordEditBox::OnGetContentForConflictPopup )
+		.OnGetMenuContent(this, &SChordEditBox::OnGetContentForConflictPopup)
 		.OnMenuOpenChanged(this, &SChordEditBox::OnConflictPopupOpenChanged)
 		[
-			SNew( SBox )
-			.WidthOverride( 200.0f )
+			SNew(SBorder)
+			.VAlign(VAlign_Center)
+			.Padding(0)
+			.BorderImage(this, &SChordEditBox::GetBorderImage)
 			[
-				SNew( SBorder )
-				.VAlign(VAlign_Center)
-				.Padding( FMargin( 4.0f, 2.0f ) )
-				.BorderImage( this, &SChordEditBox::GetBorderImage )
-				.ForegroundColor( FAppStyle::GetSlateColor(InvertedForegroundName) )
+				SNew(SHorizontalBox)
+				+ SHorizontalBox::Slot()
+				.VAlign(VAlign_Fill)
+				.HAlign(HAlign_Fill)
+				.FillWidth(1)
 				[
-					SNew( SHorizontalBox )
-					+ SHorizontalBox::Slot()
-					.FillWidth(1)
+					SNew(SBox)
+					.Padding(EditableTextBoxStyle.Padding)
 					.VAlign(VAlign_Center)
+					.WidthOverride(200.f)
 					[
-						SAssignNew( ChordEditor, SChordEditor, InputCommand, InChordIndex )
-						.OnEditBoxLostFocus( this, &SChordEditBox::OnChordEditorLostFocus )
-						.OnChordChanged( this, &SChordEditBox::OnChordChanged )
-						.OnEditingStarted( this, &SChordEditBox::OnChordEditingStarted )
-						.OnEditingStopped( this, &SChordEditBox::OnChordEditingStopped )
+						SAssignNew(ChordEditor, SChordEditor, InputCommand, InChordIndex)
+						.OnEditBoxLostFocus(this, &SChordEditBox::OnChordEditorLostFocus)
+						.OnChordChanged(this, &SChordEditBox::OnChordChanged)
+						.OnEditingStarted(this, &SChordEditBox::OnChordEditingStarted)
+						.OnEditingStopped(this, &SChordEditBox::OnChordEditingStopped)
 					]
-					+ SHorizontalBox::Slot()
-					.AutoWidth()
-					.HAlign(HAlign_Right)
-					.VAlign(VAlign_Center)
+				]
+				+ SHorizontalBox::Slot()
+				.HAlign(HAlign_Right)
+				.VAlign(VAlign_Center)
+				.AutoWidth()
+				.Padding(2.f, 0.f)
+				[
+					// Remove binding button
+					SNew(SButton)
+					.Visibility(this, &SChordEditBox::GetChordRemoveButtonVisibility)
+					.ButtonStyle(FAppStyle::Get(), "NoBorder")
+					.ContentPadding(0.f)
+					.OnClicked(this, &SChordEditBox::OnChordRemoveButtonClicked)
+					.ForegroundColor(FSlateColor::UseForeground())
+					.IsFocusable(false)
+					.ToolTipText(LOCTEXT("ChordEditButtonRemove_ToolTip", "Remove this binding"))
 					[
-						// Remove binding button
-						SNew(SButton)
-						.Visibility( this, &SChordEditBox::GetChordRemoveButtonVisibility )
-						.ButtonStyle( FAppStyle::Get(), "NoBorder" )
-						.ContentPadding(0)
-						.OnClicked( this, &SChordEditBox::OnChordRemoveButtonClicked )
-						.ForegroundColor( FSlateColor::UseForeground() )
-						.IsFocusable(false)
-						.ToolTipText(LOCTEXT("ChordEditButtonRemove_ToolTip", "Remove this binding") )
-						[
-							SNew( SImage )
-							.Image( FAppStyle::GetBrush( "Symbols.X" ) )
-							.ColorAndOpacity( FLinearColor(.7f,0,0,.75f) )
-						]
+						SNew(SImage)
+						.Image(FAppStyle::GetBrush("Symbols.X"))
+						.ColorAndOpacity(FLinearColor(.7f,0,0,.75f))
 					]
 				]
 			]
